Material Designer / Material Stream: Fixed extra menus.
- Added support for using the material stream in other menus!

diff --git a/Engine/Plugins/Experimental/DynamicMaterial/Source/DynamicMaterialEditor/Private/UI/Menus/DMMaterialSlotLayerMenus.cpp b/Engine/Plugins/Experimental/DynamicMaterial/Source/DynamicMaterialEditor/Private/UI/Menus/DMMaterialSlotLayerMenus.cpp
--- a/Engine/Plugins/Experimental/DynamicMaterial/Source/DynamicMaterialEditor/Private/UI/Menus/DMMaterialSlotLayerMenus.cpp
+++ b/Engine/Plugins/Experimental/DynamicMaterial/Source/DynamicMaterialEditor/Private/UI/Menus/DMMaterialSlotLayerMenus.cpp
@@ -59,7 +59,7 @@ TSharedRef<SWidget> FDMMaterialSlotLayerMenus::GenerateSlotLayerMenu(const TShar
 			return SNullWidget::NullWidget;
 		}
 
-		NewToolMenu->AddDynamicSection(NAME_None, FNewToolMenuDelegate::CreateStatic(&FDMMaterialSlotLayerMenus::AddAddLayerSection));
+		NewToolMenu->AddDynamicSection(TEXT("AddLayer"), FNewToolMenuDelegate::CreateStatic(&FDMMaterialSlotLayerMenus::AddAddLayerSection));
 
 		if constexpr (UE::DynamicMaterialEditor::bGlobalValuesEnabled)
 		{
@@ -216,7 +216,7 @@ void FDMMaterialSlotLayerMenus::AddAddLayerSection(UToolMenu* InMenu)
 	}
 
 	NewSection.AddSubMenu(
-		NAME_None,
+		TEXT("Advanced"),
 		LOCTEXT("AddAdvancedStage", "Advanced"),
 		LOCTEXT("AddAdvancedStageTooltip", "Add an advanced Material Stage."),
 		FNewToolMenuDelegate::CreateStatic(&AddLayerMenu_Advanced)
diff --git a/Engine/Plugins/Experimental/DynamicMaterial/Source/DynamicMaterialEditor/Private/UI/Menus/DMMaterialStageMenus.cpp b/Engine/Plugins/Experimental/DynamicMaterial/Source/DynamicMaterialEditor/Private/UI/Menus/DMMaterialStageMenus.cpp
--- a/Engine/Plugins/Experimental/DynamicMaterial/Source/DynamicMaterialEditor/Private/UI/Menus/DMMaterialStageMenus.cpp
+++ b/Engine/Plugins/Experimental/DynamicMaterial/Source/DynamicMaterialEditor/Private/UI/Menus/DMMaterialStageMenus.cpp
@@ -41,8 +41,10 @@ TSharedRef<SWidget> FDMMaterialStageMenus::GenerateStageMenu(const TSharedPtr<SD
 			return SNullWidget::NullWidget;
 		}
 
-		NewToolMenu->AddDynamicSection(NAME_None, FNewToolMenuDelegate::CreateStatic(&FDMMaterialStageMenus::AddStageSettingsSection));
-		NewToolMenu->AddDynamicSection(NAME_None, FNewToolMenuDelegate::CreateStatic(&FDMMaterialStageMenus::AddStageSourceSection));
+		//FToolMenuSection& NewSection = NewToolMenu->FindOrAddSection("ChangeMaterialStage");
+		FToolMenuSection& NewSection = NewToolMenu->FindOrAddSection(NAME_None);
+		NewSection.AddDynamicEntry(StageMenuToggleName, FNewToolMenuSectionDelegate::CreateStatic(&FDMMaterialStageMenus::AddStageSettingsSection));
+		NewSection.AddDynamicEntry(StageSourceMenuName, FNewToolMenuSectionDelegate::CreateStatic(&FDMMaterialStageMenus::AddStageSourceSection));
 	}
 
 	FToolMenuContext MenuContext(UDMMenuContext::CreateStage(InSlotWidget->GetEditorWidget(), InStageWidget));
@@ -50,16 +52,11 @@ TSharedRef<SWidget> FDMMaterialStageMenus::GenerateStageMenu(const TSharedPtr<SD
 	return ToolMenus->GenerateWidget(StageSettingsMenuName, MenuContext);
 }
 
-void FDMMaterialStageMenus::AddStageSettingsSection(UToolMenu* InMenu)
+void FDMMaterialStageMenus::AddStageSettingsSection(FToolMenuSection& InSection)
 {
 	using namespace UE::DynamicMaterialEditor::Private;
 
-	if (!IsValid(InMenu) || InMenu->ContainsSection(StageMenuToggleName))
-	{
-		return;
-	}
-
-	const UDMMenuContext* const MenuContext = InMenu->FindContext<UDMMenuContext>();
+	const UDMMenuContext* const MenuContext = InSection.FindContext<UDMMenuContext>();
 	if (!MenuContext)
 	{
 		return;
@@ -88,13 +85,13 @@ void FDMMaterialStageMenus::AddStageSettingsSection(UToolMenu* InMenu)
 	// Only if we can remove a layer can we toggle the base stage.
 	const bool bAllowRemoveLayer = Slot->CanRemoveLayer(Layer);
 
-	FToolMenuSection& NewSection = InMenu->AddSection(StageMenuToggleName, LOCTEXT("MaterialStageMenu", "Stage Actions"));
+	InSection.InitSection(StageMenuToggleName, LOCTEXT("MaterialStageMenu", "Stage Actions"), FToolMenuInsert());
 
 	if (bAllowRemoveLayer)
 	{
 		if (StageType == EDMMaterialLayerStage::Base)
 		{
-			NewSection.AddMenuEntry(NAME_None,
+			InSection.AddMenuEntry(NAME_None,
 				LOCTEXT("ToggleLayerBase", "Toggle Base"),
 				LOCTEXT("ToggleLayerBaseTooltip", "Toggle the Layer Base.\n\nAlt+Shift+Left Click"),
 				FSlateIcon(FDynamicMaterialEditorStyle::Get().GetStyleSetName(), TEXT("Icons.Stage.Enabled")),
@@ -117,7 +114,7 @@ void FDMMaterialStageMenus::AddStageSettingsSection(UToolMenu* InMenu)
 
 	if (StageType == EDMMaterialLayerStage::Mask)
 	{
-		NewSection.AddMenuEntry(NAME_None,
+		InSection.AddMenuEntry(NAME_None,
 			LOCTEXT("ToggleLayerMask", "Toggle Mask"),
 			LOCTEXT("ToggleLayerMaskTooltip", "Toggle the Layer Mask.\n\nAlt+Shift+Left Click"),
 			FSlateIcon(FDynamicMaterialEditorStyle::Get().GetStyleSetName(), TEXT("Icons.Stage.Enabled")),
@@ -138,19 +135,13 @@ void FDMMaterialStageMenus::AddStageSettingsSection(UToolMenu* InMenu)
 	}
 }
 
-void FDMMaterialStageMenus::AddStageSourceSection(UToolMenu* InMenu)
+void FDMMaterialStageMenus::AddStageSourceSection(FToolMenuSection& InSection)
 {
 	using namespace UE::DynamicMaterialEditor::Private;
 
-	if (!IsValid(InMenu) || InMenu->ContainsSection(StageSourceMenuName))
-	{
-		return;
-	}
-
-	FToolMenuSection& NewSection = InMenu->AddSection(StageSourceMenuName, LOCTEXT("MaterialStageSource", "Change Stage Source"));
-	NewSection.Context.AddObject(InMenu->Context.FindContext<UDMMenuContext>());
+	InSection.InitSection(StageSourceMenuName, LOCTEXT("MaterialStageSource", "Change Stage Source"), FToolMenuInsert());
 
-	FDMMaterialStageSourceMenus::CreateChangeMaterialStageSource(NewSection);
+	FDMMaterialStageSourceMenus::CreateChangeMaterialStageSource(InSection);
 }
 
 #undef LOCTEXT_NAMESPACE
diff --git a/Engine/Plugins/Experimental/DynamicMaterial/Source/DynamicMaterialEditor/Private/UI/Menus/DMMaterialStageMenus.h b/Engine/Plugins/Experimental/DynamicMaterial/Source/DynamicMaterialEditor/Private/UI/Menus/DMMaterialStageMenus.h
--- a/Engine/Plugins/Experimental/DynamicMaterial/Source/DynamicMaterialEditor/Private/UI/Menus/DMMaterialStageMenus.h
+++ b/Engine/Plugins/Experimental/DynamicMaterial/Source/DynamicMaterialEditor/Private/UI/Menus/DMMaterialStageMenus.h
@@ -8,7 +8,7 @@ class FName;
 class SDMMaterialSlotEditor;
 class SDMMaterialStage;
 class SWidget;
-class UToolMenu;
+struct FToolMenuSection;
 
 class FDMMaterialStageMenus final
 {
@@ -16,7 +16,7 @@ public:
 	static TSharedRef<SWidget> GenerateStageMenu(const TSharedPtr<SDMMaterialSlotEditor>& InSlotWidget, const TSharedPtr<SDMMaterialStage>& InStageWidget);
 
 private:
-	static void AddStageSettingsSection(UToolMenu* InMenu);
+	static void AddStageSettingsSection(FToolMenuSection& InSection);
 
-	static void AddStageSourceSection(UToolMenu* InMenu);
+	static void AddStageSourceSection(FToolMenuSection& InSection);
 };
diff --git a/Engine/Plugins/Experimental/DynamicMaterialMediaStreamBridge/Source/DynamicMaterialMediaStreamBridgeEditor/Private/DMMediaStreamStageSourceMenuExtender.cpp b/Engine/Plugins/Experimental/DynamicMaterialMediaStreamBridge/Source/DynamicMaterialMediaStreamBridgeEditor/Private/DMMediaStreamStageSourceMenuExtender.cpp
--- a/Engine/Plugins/Experimental/DynamicMaterialMediaStreamBridge/Source/DynamicMaterialMediaStreamBridgeEditor/Private/DMMediaStreamStageSourceMenuExtender.cpp
+++ b/Engine/Plugins/Experimental/DynamicMaterialMediaStreamBridge/Source/DynamicMaterialMediaStreamBridgeEditor/Private/DMMediaStreamStageSourceMenuExtender.cpp
@@ -2,6 +2,8 @@
 
 #include "DMMediaStreamStageSourceMenuExtender.h"
 
+#include "Components/DMMaterialLayer.h"
+#include "Components/DMMaterialSlot.h"
 #include "Components/DMMaterialStage.h"
 #include "Components/DMMaterialStageBlend.h"
 #include "Components/DMMaterialStageThroughputLayerBlend.h"
@@ -17,7 +19,7 @@
 
 #define LOCTEXT_NAMESPACE "DMMEdiaStreamStageSourceMenuExtender"
 
-void FDMMEdiaStreamStageSourceMenuExtender::Integrate()
+void FDMMediaStreamStageSourceMenuExtender::Integrate()
 {
 	if (bIntegrated)
 	{
@@ -31,22 +33,34 @@ void FDMMEdiaStreamStageSourceMenuExtender::Integrate()
 		return;
 	}
 
-	UToolMenu* SourceMenu = ToolMenus->ExtendMenu(TEXT("MaterialDesigner.MaterialStage.ChangeSource.Advanced"));
+	if (UToolMenu* AdvancedMenu = ToolMenus->ExtendMenu(TEXT("MaterialDesigner.MaterialStage.ChangeSource.Advanced")))
+	{
+		AdvancedMenu->AddDynamicSection(
+			NAME_None,
+			FNewToolMenuDelegate::CreateRaw(this, &FDMMediaStreamStageSourceMenuExtender::ExtendMenu_ChangeSource)
+		);
+	}
 
-	if (!SourceMenu)
+	if (UToolMenu* AdvancedMenu = ToolMenus->ExtendMenu(TEXT("MaterialDesigner.MaterialStage.Advanced")))
 	{
-		return;
+		AdvancedMenu->AddDynamicSection(
+			NAME_None,
+			FNewToolMenuDelegate::CreateRaw(this, &FDMMediaStreamStageSourceMenuExtender::ExtendMenu_ChangeSource)
+		);
 	}
 
-	SourceMenu->AddDynamicSection(
-		TEXT("Advanced"), 
-		FNewToolMenuDelegate::CreateRaw(this, &FDMMEdiaStreamStageSourceMenuExtender::ExtendMenu)
-	);
+	if (UToolMenu* AdvancedMenu = ToolMenus->ExtendMenu(TEXT("MaterialDesigner.MaterialSlot.Layer.Advanced")))
+	{
+		AdvancedMenu->AddDynamicSection(
+			NAME_None,
+			FNewToolMenuDelegate::CreateRaw(this, &FDMMediaStreamStageSourceMenuExtender::ExtendMenu_AddLayer)
+		);
+	}
 
 	bIntegrated = true;
 }
 
-void FDMMEdiaStreamStageSourceMenuExtender::ExtendMenu(UToolMenu* InMenu)
+void FDMMediaStreamStageSourceMenuExtender::ExtendMenu_ChangeSource(UToolMenu* InMenu)
 {
 	UDMMenuContext* MenuContext = InMenu->FindContext<UDMMenuContext>();
 
@@ -55,23 +69,48 @@ void FDMMEdiaStreamStageSourceMenuExtender::ExtendMenu(UToolMenu* InMenu)
 		return;
 	}
 
-	FToolMenuSection& AdvancedSection = InMenu->AddSection(NAME_None);
+	FToolMenuSection& AdvancedSection = InMenu->FindOrAddSection(NAME_None);
+	AdvancedSection.Context.AddObject(MenuContext);
 
 	AdvancedSection.AddMenuEntry("MediaStream",
 		LOCTEXT("ChangeSourceMediaStream", "Media Stream"),
 		LOCTEXT("ChangeSourceMediaStreamTooltip", "Change the source of this stage to a Media Stream."),
 		GetDefault<UDMMaterialValueMediaStream>()->GetComponentIcon(),
 		FUIAction(
-			FExecuteAction::CreateRaw(
-				this,
-				&FDMMEdiaStreamStageSourceMenuExtender::ChangeSourceToMediaStreamFromContext,
+			FExecuteAction::CreateRaw(this,
+				&FDMMediaStreamStageSourceMenuExtender::ChangeSourceToMediaStreamFromContext,
 				MenuContext
 			)
 		)
 	);
 }
 
-void FDMMEdiaStreamStageSourceMenuExtender::ChangeSourceToMediaStreamFromContext(UDMMenuContext* InMenuContext)
+void FDMMediaStreamStageSourceMenuExtender::ExtendMenu_AddLayer(UToolMenu* InMenu)
+{
+	UDMMenuContext* MenuContext = InMenu->FindContext<UDMMenuContext>();
+
+	if (!ensure(MenuContext))
+	{
+		return;
+	}
+
+	FToolMenuSection& AdvancedSection = InMenu->FindOrAddSection(NAME_None);
+	AdvancedSection.Context.AddObject(MenuContext);
+
+	AdvancedSection.AddMenuEntry("MediaStream",
+		LOCTEXT("ChangeSourceMediaStream", "Media Stream"),
+		LOCTEXT("ChangeSourceMediaStreamTooltip", "Add a new layer based on a Media Stream."),
+		GetDefault<UDMMaterialValueMediaStream>()->GetComponentIcon(),
+		FUIAction(
+			FExecuteAction::CreateRaw(this,
+				&FDMMediaStreamStageSourceMenuExtender::AddMediaStreamLayerFromContext,
+				MenuContext
+			)
+		)
+	);
+}
+
+void FDMMediaStreamStageSourceMenuExtender::ChangeSourceToMediaStreamFromContext(UDMMenuContext* InMenuContext)
 {
 	if (!IsValid(InMenuContext))
 	{
@@ -143,9 +182,86 @@ void FDMMEdiaStreamStageSourceMenuExtender::ChangeSourceToMediaStreamFromContext
 	);
 }
 
-FDMMEdiaStreamStageSourceMenuExtender& FDMMEdiaStreamStageSourceMenuExtender::Get()
+void FDMMediaStreamStageSourceMenuExtender::AddMediaStreamLayerFromContext(UDMMenuContext* InMenuContext)
+{
+	if (!IsValid(InMenuContext))
+	{
+		return;
+	}
+
+	UDMMaterialSlot* Slot = InMenuContext->GetSlot();
+
+	if (!Slot)
+	{
+		return;
+	}
+
+	UDMMaterialLayerObject* OldLayer = InMenuContext->GetLayer();
+
+	if (!OldLayer)
+	{
+		return;
+	}
+
+	if (GUndo)
+	{
+		Slot->Modify();
+	}
+
+	UDMMaterialLayerObject* NewLayer = nullptr;
+
+	const FDMUpdateGuard Guard;
+
+	NewLayer = Slot->AddDefaultLayer(OldLayer->GetMaterialProperty());
+
+	if (!ensure(NewLayer))
+	{
+		return;
+	}
+
+	UDMMaterialStage* Stage = NewLayer->GetStage(EDMMaterialLayerStage::Base);
+
+	if (!ensure(Stage))
+	{
+		return;
+	}
+
+	UDMMaterialStageInputExpression* NewExpression = UDMMaterialStageInputExpression::ChangeStageInput_Expression(
+		Stage,
+		UDMMaterialStageExpressionTextureSample::StaticClass(),
+		UDMMaterialStageBlend::InputB,
+		FDMMaterialStageConnectorChannel::WHOLE_CHANNEL,
+		0,
+		FDMMaterialStageConnectorChannel::WHOLE_CHANNEL
+	);
+
+	if (!ensure(NewExpression))
+	{
+		return;
+	}
+
+	UDMMaterialSubStage* SubStage = NewExpression->GetSubStage();
+
+	if (!ensure(SubStage))
+	{
+		return;
+	}
+
+	UDMMaterialStageInputValue::ChangeStageInput_NewLocalValue(
+		SubStage,
+		0,
+		FDMMaterialStageConnectorChannel::WHOLE_CHANNEL,
+		UDMMaterialValueMediaStream::StaticClass(),
+		FDMMaterialStageConnectorChannel::WHOLE_CHANNEL
+	);
+
+	NewLayer->Update(NewLayer, EDMUpdateType::Structure);
+
+}
+
+FDMMediaStreamStageSourceMenuExtender& FDMMediaStreamStageSourceMenuExtender::Get()
 {
-	static FDMMEdiaStreamStageSourceMenuExtender MenuExtension;
+	static FDMMediaStreamStageSourceMenuExtender MenuExtension;
 	return MenuExtension;
 }
 
diff --git a/Engine/Plugins/Experimental/DynamicMaterialMediaStreamBridge/Source/DynamicMaterialMediaStreamBridgeEditor/Private/DMMediaStreamStageSourceMenuExtender.h b/Engine/Plugins/Experimental/DynamicMaterialMediaStreamBridge/Source/DynamicMaterialMediaStreamBridgeEditor/Private/DMMediaStreamStageSourceMenuExtender.h
--- a/Engine/Plugins/Experimental/DynamicMaterialMediaStreamBridge/Source/DynamicMaterialMediaStreamBridgeEditor/Private/DMMediaStreamStageSourceMenuExtender.h
+++ b/Engine/Plugins/Experimental/DynamicMaterialMediaStreamBridge/Source/DynamicMaterialMediaStreamBridgeEditor/Private/DMMediaStreamStageSourceMenuExtender.h
@@ -4,18 +4,23 @@
 
 class UDMMenuContext;
 class UToolMenu;
+struct FToolMenuSection;
 
-class FDMMEdiaStreamStageSourceMenuExtender
+class FDMMediaStreamStageSourceMenuExtender
 {
 public:
-	static FDMMEdiaStreamStageSourceMenuExtender& Get();
+	static FDMMediaStreamStageSourceMenuExtender& Get();
 
 	void Integrate();
 
 private:
 	bool bIntegrated = false;
 
-	void ExtendMenu(UToolMenu* InMenu);
+	void ExtendMenu_ChangeSource(UToolMenu* InMenu);
+
+	void ExtendMenu_AddLayer(UToolMenu* InMenu);
 
 	void ChangeSourceToMediaStreamFromContext(UDMMenuContext* InMenuContext);
+
+	void AddMediaStreamLayerFromContext(UDMMenuContext* InMenuContext);
 };
diff --git a/Engine/Plugins/Experimental/DynamicMaterialMediaStreamBridge/Source/DynamicMaterialMediaStreamBridgeEditor/Private/DynamicMaterialMediaStreamBridgeEditorModule.cpp b/Engine/Plugins/Experimental/DynamicMaterialMediaStreamBridge/Source/DynamicMaterialMediaStreamBridgeEditor/Private/DynamicMaterialMediaStreamBridgeEditorModule.cpp
--- a/Engine/Plugins/Experimental/DynamicMaterialMediaStreamBridge/Source/DynamicMaterialMediaStreamBridgeEditor/Private/DynamicMaterialMediaStreamBridgeEditorModule.cpp
+++ b/Engine/Plugins/Experimental/DynamicMaterialMediaStreamBridge/Source/DynamicMaterialMediaStreamBridgeEditor/Private/DynamicMaterialMediaStreamBridgeEditorModule.cpp
@@ -13,7 +13,7 @@ void FDynamicMaterialMediaStreamBridgeEditorModule::StartupModule()
 	IDynamicMaterialEditorModule& DMEditorModule = IDynamicMaterialEditorModule::Get();
 	DMEditorModule.RegisterComponentPropertyRowGeneratorDelegate<UDMMaterialValueMediaStream, FDMMaterialValueMediaStreamPropertyRowGenerator>();
 
-	FDMMEdiaStreamStageSourceMenuExtender::Get().Integrate();
+	FDMMediaStreamStageSourceMenuExtender::Get().Integrate();
 }
 
 IMPLEMENT_MODULE(FDynamicMaterialMediaStreamBridgeEditorModule, DynamicMaterialMediaStreamBridgeEditor)
