Media Stream: Player change delegate
- Renamed the OnPlayerChanged delegate to OnSourceChanged.
- Made the delegate fire whenever ApplySource is called (before it was not called if the current player was null.)

diff --git a/Engine/Plugins/Experimental/DynamicMaterialMediaStreamBridge/Source/DynamicMaterialMediaStreamBridge/Private/DMMaterialValueMediaStream.cpp b/Engine/Plugins/Experimental/DynamicMaterialMediaStreamBridge/Source/DynamicMaterialMediaStreamBridge/Private/DMMaterialValueMediaStream.cpp
--- a/Engine/Plugins/Experimental/DynamicMaterialMediaStreamBridge/Source/DynamicMaterialMediaStreamBridge/Private/DMMaterialValueMediaStream.cpp
+++ b/Engine/Plugins/Experimental/DynamicMaterialMediaStreamBridge/Source/DynamicMaterialMediaStreamBridge/Private/DMMaterialValueMediaStream.cpp
@@ -181,7 +181,7 @@ void UDMMaterialValueMediaStream::SubscribeToPlayerChange()
 {
 	if (IsValid(MediaStream))
 	{
-		MediaStream->GetOnPlayerChanged().AddDynamic(this, &UDMMaterialValueMediaStream::OnPlayerChanged);
+		MediaStream->GetOnSourceChanged().AddDynamic(this, &UDMMaterialValueMediaStream::OnSourceChanged);
 		SubscribedStreamWeak = MediaStream;
 	}
 	else
@@ -194,16 +194,16 @@ void UDMMaterialValueMediaStream::UnsubscribeFromPlayerChange()
 {
 	if (UMediaStream* SubscribedStream = SubscribedStreamWeak.Get())
 	{
-		SubscribedStream->GetOnPlayerChanged().RemoveAll(this);
+		SubscribedStream->GetOnSourceChanged().RemoveAll(this);
 	}
 
 	if (IsValid(MediaStream))
 	{
-		MediaStream->GetOnPlayerChanged().RemoveAll(this);
+		MediaStream->GetOnSourceChanged().RemoveAll(this);
 	}
 }
 
-void UDMMaterialValueMediaStream::OnPlayerChanged(UMediaStream* InMediaStream, TScriptInterface<IMediaStreamPlayer> InPlayer)
+void UDMMaterialValueMediaStream::OnSourceChanged(UMediaStream* InMediaStream)
 {
 	UpdatePlayer();
 
diff --git a/Engine/Plugins/Experimental/DynamicMaterialMediaStreamBridge/Source/DynamicMaterialMediaStreamBridge/Private/DMMaterialValueMediaStreamDynamic.cpp b/Engine/Plugins/Experimental/DynamicMaterialMediaStreamBridge/Source/DynamicMaterialMediaStreamBridge/Private/DMMaterialValueMediaStreamDynamic.cpp
--- a/Engine/Plugins/Experimental/DynamicMaterialMediaStreamBridge/Source/DynamicMaterialMediaStreamBridge/Private/DMMaterialValueMediaStreamDynamic.cpp
+++ b/Engine/Plugins/Experimental/DynamicMaterialMediaStreamBridge/Source/DynamicMaterialMediaStreamBridge/Private/DMMaterialValueMediaStreamDynamic.cpp
@@ -155,7 +155,7 @@ void UDMMaterialValueMediaStreamDynamic::SubscribeToPlayerChange()
 {
 	if (IsValid(MediaStream))
 	{
-		MediaStream->GetOnPlayerChanged().AddDynamic(this, &UDMMaterialValueMediaStreamDynamic::OnPlayerChanged);
+		MediaStream->GetOnSourceChanged().AddDynamic(this, &UDMMaterialValueMediaStreamDynamic::OnSourceChanged);
 		SubscribedStreamWeak = MediaStream;
 	}
 	else
@@ -168,16 +168,16 @@ void UDMMaterialValueMediaStreamDynamic::UnsubscribeFromPlayerChange()
 {
 	if (UMediaStream* SubscribedStream = SubscribedStreamWeak.Get())
 	{
-		SubscribedStream->GetOnPlayerChanged().RemoveAll(this);
+		SubscribedStream->GetOnSourceChanged().RemoveAll(this);
 	}
 
 	if (IsValid(MediaStream))
 	{
-		MediaStream->GetOnPlayerChanged().RemoveAll(this);
+		MediaStream->GetOnSourceChanged().RemoveAll(this);
 	}
 }
 
-void UDMMaterialValueMediaStreamDynamic::OnPlayerChanged(UMediaStream* InMediaStream, TScriptInterface<IMediaStreamPlayer> InPlayer)
+void UDMMaterialValueMediaStreamDynamic::OnSourceChanged(UMediaStream* InMediaStream)
 {
 	UpdatePlayer();
 
diff --git a/Engine/Plugins/Experimental/DynamicMaterialMediaStreamBridge/Source/DynamicMaterialMediaStreamBridge/Public/DMMaterialValueMediaStream.h b/Engine/Plugins/Experimental/DynamicMaterialMediaStreamBridge/Source/DynamicMaterialMediaStreamBridge/Public/DMMaterialValueMediaStream.h
--- a/Engine/Plugins/Experimental/DynamicMaterialMediaStreamBridge/Source/DynamicMaterialMediaStreamBridge/Public/DMMaterialValueMediaStream.h
+++ b/Engine/Plugins/Experimental/DynamicMaterialMediaStreamBridge/Source/DynamicMaterialMediaStreamBridge/Public/DMMaterialValueMediaStream.h
@@ -63,7 +63,7 @@ protected:
 	void UnsubscribeFromPlayerChange();
 
 	UFUNCTION()
-	void OnPlayerChanged(UMediaStream* InMediaStream, TScriptInterface<IMediaStreamPlayer> InPlayer);
+	void OnSourceChanged(UMediaStream* InMediaStream);
 
 	void UpdatePlayer();
 
diff --git a/Engine/Plugins/Experimental/DynamicMaterialMediaStreamBridge/Source/DynamicMaterialMediaStreamBridge/Public/DMMaterialValueMediaStreamDynamic.h b/Engine/Plugins/Experimental/DynamicMaterialMediaStreamBridge/Source/DynamicMaterialMediaStreamBridge/Public/DMMaterialValueMediaStreamDynamic.h
--- a/Engine/Plugins/Experimental/DynamicMaterialMediaStreamBridge/Source/DynamicMaterialMediaStreamBridge/Public/DMMaterialValueMediaStreamDynamic.h
+++ b/Engine/Plugins/Experimental/DynamicMaterialMediaStreamBridge/Source/DynamicMaterialMediaStreamBridge/Public/DMMaterialValueMediaStreamDynamic.h
@@ -60,7 +60,7 @@ protected:
 	void UnsubscribeFromPlayerChange();
 
 	UFUNCTION()
-	void OnPlayerChanged(UMediaStream* InMediaStream, TScriptInterface<IMediaStreamPlayer> InPlayer);
+	void OnSourceChanged(UMediaStream* InMediaStream);
 
 	void UpdatePlayer();
 
diff --git a/Engine/Plugins/Experimental/MediaStream/Source/MediaStream/Private/MediaStream.cpp b/Engine/Plugins/Experimental/MediaStream/Source/MediaStream/Private/MediaStream.cpp
--- a/Engine/Plugins/Experimental/MediaStream/Source/MediaStream/Private/MediaStream.cpp
+++ b/Engine/Plugins/Experimental/MediaStream/Source/MediaStream/Private/MediaStream.cpp
@@ -50,9 +50,9 @@ TScriptInterface<IMediaStreamPlayer> UMediaStream::GetPlayer() const
 	return PlayerObject;
 }
 
-UMediaStream::FOnPlayerChanged& UMediaStream::GetOnPlayerChanged()
+UMediaStream::FOnSourceChanged& UMediaStream::GetOnSourceChanged()
 {
-	return OnPlayerChanged;
+	return OnSourceChanged;
 }
 
 #if WITH_EDITOR
@@ -220,10 +220,11 @@ bool UMediaStream::ApplySource()
 				Player->Deinitialize();
 				Player->OnSourceChanged(Source);
 
-				OnPlayerChanged.Broadcast(this, PlayerObject);
 			}
 		}
 
+		OnSourceChanged.Broadcast(this);
+
 		return false;
 	}
 
@@ -271,7 +272,7 @@ bool UMediaStream::ApplySource()
 		Player->OnSourceChanged(Source);
 	}
 
-	OnPlayerChanged.Broadcast(this, PlayerObject);
+	OnSourceChanged.Broadcast(this);
 
 	return true;
 }
diff --git a/Engine/Plugins/Experimental/MediaStream/Source/MediaStream/Private/MediaStreamComponent.cpp b/Engine/Plugins/Experimental/MediaStream/Source/MediaStream/Private/MediaStreamComponent.cpp
--- a/Engine/Plugins/Experimental/MediaStream/Source/MediaStream/Private/MediaStreamComponent.cpp
+++ b/Engine/Plugins/Experimental/MediaStream/Source/MediaStream/Private/MediaStreamComponent.cpp
@@ -49,7 +49,7 @@ void UMediaStreamComponent::PostInitProperties()
 
 	SetStaticMesh(UE::MediaStream::Private::GetStaticMesh());
 
-	MediaStream->GetOnPlayerChanged().AddDynamic(this, &UMediaStreamComponent::OnPlayerChanged);
+	MediaStream->GetOnSourceChanged().AddDynamic(this, &UMediaStreamComponent::OnSourceChanged);
 
 	InitPlayer();
 }
@@ -68,7 +68,7 @@ void UMediaStreamComponent::PostNetReceive()
 	InitPlayer();
 }
 
-void UMediaStreamComponent::OnPlayerChanged(UMediaStream* InMediaStream, TScriptInterface<IMediaStreamPlayer> InPlayer)
+void UMediaStreamComponent::OnSourceChanged(UMediaStream* InMediaStream)
 {
 	InitPlayer();
 }
diff --git a/Engine/Plugins/Experimental/MediaStream/Source/MediaStream/Public/MediaStream.h b/Engine/Plugins/Experimental/MediaStream/Source/MediaStream/Public/MediaStream.h
--- a/Engine/Plugins/Experimental/MediaStream/Source/MediaStream/Public/MediaStream.h
+++ b/Engine/Plugins/Experimental/MediaStream/Source/MediaStream/Public/MediaStream.h
@@ -59,8 +59,8 @@ public:
 	/**
 	 * Called when the Media Stream Player changes.
 	 */
-	DECLARE_DYNAMIC_MULTICAST_DELEGATE_TwoParams(FOnPlayerChanged, UMediaStream*, InMediaStream, TScriptInterface<IMediaStreamPlayer>, InPlayer);
-	MEDIASTREAM_API FOnPlayerChanged& GetOnPlayerChanged();
+	DECLARE_DYNAMIC_MULTICAST_DELEGATE_OneParam(FOnSourceChanged, UMediaStream*, InMediaStream);
+	MEDIASTREAM_API FOnSourceChanged& GetOnSourceChanged();
 
 	//~ Begin UObject
 #if WITH_EDITOR
@@ -98,7 +98,7 @@ protected:
 	TObjectPtr<UObject> PlayerObject = nullptr;
 
 	UPROPERTY(BlueprintAssignable)
-	FOnPlayerChanged OnPlayerChanged;
+	FOnSourceChanged OnSourceChanged;
 
 	FMediaStreamSource Source_PreUndo;
 
diff --git a/Engine/Plugins/Experimental/MediaStream/Source/MediaStream/Public/MediaStreamComponent.h b/Engine/Plugins/Experimental/MediaStream/Source/MediaStream/Public/MediaStreamComponent.h
--- a/Engine/Plugins/Experimental/MediaStream/Source/MediaStream/Public/MediaStreamComponent.h
+++ b/Engine/Plugins/Experimental/MediaStream/Source/MediaStream/Public/MediaStreamComponent.h
@@ -31,7 +31,7 @@ protected:
 	TObjectPtr<UMaterialInstanceDynamic> MaterialInstance;
 
 	UFUNCTION()
-	void OnPlayerChanged(UMediaStream* InMediaStream, TScriptInterface<IMediaStreamPlayer> InPlayer);
+	void OnSourceChanged(UMediaStream* InMediaStream);
 
 	void InitPlayer();
 };
