Horde Issue 747399

diff --git a/Engine/Plugins/Experimental/DynamicMaterialMediaStreamBridge/Source/DynamicMaterialMediaStreamBridgeEditor/Private/DMMediaStreamStageSourceMenuExtender.cpp b/Engine/Plugins/Experimental/DynamicMaterialMediaStreamBridge/Source/DynamicMaterialMediaStreamBridgeEditor/Private/DMMediaStreamStageSourceMenuExtender.cpp
--- a/Engine/Plugins/Experimental/DynamicMaterialMediaStreamBridge/Source/DynamicMaterialMediaStreamBridgeEditor/Private/DMMediaStreamStageSourceMenuExtender.cpp
+++ b/Engine/Plugins/Experimental/DynamicMaterialMediaStreamBridge/Source/DynamicMaterialMediaStreamBridgeEditor/Private/DMMediaStreamStageSourceMenuExtender.cpp
@@ -98,8 +98,8 @@ void FDMMediaStreamStageSourceMenuExtender::ExtendMenu_AddLayer(UToolMenu* InMen
 	AdvancedSection.Context.AddObject(MenuContext);
 
 	AdvancedSection.AddMenuEntry("MediaStream",
-		LOCTEXT("ChangeSourceMediaStream", "Media Stream"),
-		LOCTEXT("ChangeSourceMediaStreamTooltip", "Add a new layer based on a Media Stream."),
+		LOCTEXT("AddLayerMediaStream", "Media Stream"),
+		LOCTEXT("AddLayerMediaStreamTooltip", "Add a new layer based on a Media Stream."),
 		GetDefault<UDMMaterialValueMediaStream>()->GetComponentIcon(),
 		FUIAction(
 			FExecuteAction::CreateRaw(this,
