Material Designer / Media Stream: Integration update
- Fixed add layer with media stream.
- Renamed UI instances of Media Stream to Video

diff --git a/Engine/Plugins/Experimental/DynamicMaterialMediaStreamBridge/Source/DynamicMaterialMediaStreamBridge/Private/DMMaterialValueMediaStream.cpp b/Engine/Plugins/Experimental/DynamicMaterialMediaStreamBridge/Source/DynamicMaterialMediaStreamBridge/Private/DMMaterialValueMediaStream.cpp
--- a/Engine/Plugins/Experimental/DynamicMaterialMediaStreamBridge/Source/DynamicMaterialMediaStreamBridge/Private/DMMaterialValueMediaStream.cpp
+++ b/Engine/Plugins/Experimental/DynamicMaterialMediaStreamBridge/Source/DynamicMaterialMediaStreamBridge/Private/DMMaterialValueMediaStream.cpp
@@ -138,7 +138,7 @@ FString UDMMaterialValueMediaStream::GetComponentPathComponent() const
 
 FText UDMMaterialValueMediaStream::GetComponentDescription() const
 {
-	return LOCTEXT("MediaStream", "Media Stream");
+	return LOCTEXT("MediaStream", "Video");
 }
 
 void UDMMaterialValueMediaStream::PostEditorDuplicate(UDynamicMaterialModel* InMaterialModel, UDMMaterialComponent* InParent)
diff --git a/Engine/Plugins/Experimental/DynamicMaterialMediaStreamBridge/Source/DynamicMaterialMediaStreamBridge/Private/DMMaterialValueMediaStreamDynamic.cpp b/Engine/Plugins/Experimental/DynamicMaterialMediaStreamBridge/Source/DynamicMaterialMediaStreamBridge/Private/DMMaterialValueMediaStreamDynamic.cpp
--- a/Engine/Plugins/Experimental/DynamicMaterialMediaStreamBridge/Source/DynamicMaterialMediaStreamBridge/Private/DMMaterialValueMediaStreamDynamic.cpp
+++ b/Engine/Plugins/Experimental/DynamicMaterialMediaStreamBridge/Source/DynamicMaterialMediaStreamBridge/Private/DMMaterialValueMediaStreamDynamic.cpp
@@ -112,7 +112,7 @@ FString UDMMaterialValueMediaStreamDynamic::GetComponentPathComponent() const
 
 FText UDMMaterialValueMediaStreamDynamic::GetComponentDescription() const
 {
-	return LOCTEXT("MediaStream", "Media Stream");
+	return LOCTEXT("MediaStream", "Video");
 }
 
 void UDMMaterialValueMediaStreamDynamic::PostEditorDuplicate(UDynamicMaterialModel* InMaterialModel, UDMMaterialComponent* InParent)
diff --git a/Engine/Plugins/Experimental/DynamicMaterialMediaStreamBridge/Source/DynamicMaterialMediaStreamBridgeEditor/Private/DMMaterialValueMediaStreamPropertyRowGenerator.cpp b/Engine/Plugins/Experimental/DynamicMaterialMediaStreamBridge/Source/DynamicMaterialMediaStreamBridgeEditor/Private/DMMaterialValueMediaStreamPropertyRowGenerator.cpp
--- a/Engine/Plugins/Experimental/DynamicMaterialMediaStreamBridge/Source/DynamicMaterialMediaStreamBridgeEditor/Private/DMMaterialValueMediaStreamPropertyRowGenerator.cpp
+++ b/Engine/Plugins/Experimental/DynamicMaterialMediaStreamBridge/Source/DynamicMaterialMediaStreamBridgeEditor/Private/DMMaterialValueMediaStreamPropertyRowGenerator.cpp
@@ -70,20 +70,20 @@ void FDMMaterialValueMediaStreamPropertyRowGenerator::AddPlayerProperty(const TS
 	FDMPropertyHandle& ControlHandle = InOutPropertyRows.Add_GetRef(
 		IDMWidgetLibrary::Get().GetPropertyHandle(&*InComponentEditorWidget, InMediaStream, TEXT("Control"))
 	);
-	ControlHandle.CategoryOverrideName = TEXT("Media Stream Player");
+	ControlHandle.CategoryOverrideName = TEXT("Video Player");
 
 	FDMPropertyHandle& TextureHandle = InOutPropertyRows.Add_GetRef(
 		IDMWidgetLibrary::Get().GetPropertyHandle(&*InComponentEditorWidget, InMediaStream, TEXT("MediaTexture"))
 	);
-	TextureHandle.CategoryOverrideName = TEXT("Media Stream Player");
+	TextureHandle.CategoryOverrideName = TEXT("Video Player");
 
 	FDMPropertyHandle& TextureConfigHandle = InOutPropertyRows.Add_GetRef(
 		IDMWidgetLibrary::Get().GetPropertyHandle(&*InComponentEditorWidget, InMediaStream, TEXT("TextureConfig"))
 	);
-	TextureConfigHandle.CategoryOverrideName = TEXT("Media Stream Player");
+	TextureConfigHandle.CategoryOverrideName = TEXT("Video Player");
 
 	FDMPropertyHandle& PlayerConfigHandle = InOutPropertyRows.Add_GetRef(
 		IDMWidgetLibrary::Get().GetPropertyHandle(&*InComponentEditorWidget, InMediaStream, TEXT("PlayerConfig"))
 	);
-	PlayerConfigHandle.CategoryOverrideName = TEXT("Media Stream Player");
+	PlayerConfigHandle.CategoryOverrideName = TEXT("Video Player");
 }
diff --git a/Engine/Plugins/Experimental/DynamicMaterialMediaStreamBridge/Source/DynamicMaterialMediaStreamBridgeEditor/Private/DMMediaStreamStageSourceMenuExtender.cpp b/Engine/Plugins/Experimental/DynamicMaterialMediaStreamBridge/Source/DynamicMaterialMediaStreamBridgeEditor/Private/DMMediaStreamStageSourceMenuExtender.cpp
--- a/Engine/Plugins/Experimental/DynamicMaterialMediaStreamBridge/Source/DynamicMaterialMediaStreamBridgeEditor/Private/DMMediaStreamStageSourceMenuExtender.cpp
+++ b/Engine/Plugins/Experimental/DynamicMaterialMediaStreamBridge/Source/DynamicMaterialMediaStreamBridgeEditor/Private/DMMediaStreamStageSourceMenuExtender.cpp
@@ -15,6 +15,8 @@
 #include "ScopedTransaction.h"
 #include "ToolMenu.h"
 #include "ToolMenus.h"
+#include "Components/DMMaterialValueDynamic.h"
+#include "Model/DynamicMaterialModelEditorOnlyData.h"
 #include "UI/Menus/DMMenuContext.h"
 
 #define LOCTEXT_NAMESPACE "DMMEdiaStreamStageSourceMenuExtender"
@@ -73,8 +75,8 @@ void FDMMediaStreamStageSourceMenuExtender::ExtendMenu_ChangeSource(UToolMenu* I
 	AdvancedSection.Context.AddObject(MenuContext);
 
 	AdvancedSection.AddMenuEntry("MediaStream",
-		LOCTEXT("ChangeSourceMediaStream", "Media Stream"),
-		LOCTEXT("ChangeSourceMediaStreamTooltip", "Change the source of this stage to a Media Stream."),
+		LOCTEXT("ChangeSourceMediaStream", "Video"),
+		LOCTEXT("ChangeSourceMediaStreamTooltip", "Change the source of this stage to a Video."),
 		GetDefault<UDMMaterialValueMediaStream>()->GetComponentIcon(),
 		FUIAction(
 			FExecuteAction::CreateRaw(this,
@@ -98,8 +100,8 @@ void FDMMediaStreamStageSourceMenuExtender::ExtendMenu_AddLayer(UToolMenu* InMen
 	AdvancedSection.Context.AddObject(MenuContext);
 
 	AdvancedSection.AddMenuEntry("MediaStream",
-		LOCTEXT("AddLayerMediaStream", "Media Stream"),
-		LOCTEXT("AddLayerMediaStreamTooltip", "Add a new layer based on a Media Stream."),
+		LOCTEXT("AddLayerMediaStream", "Video"),
+		LOCTEXT("AddLayerMediaStreamTooltip", "Add a new layer based on a Video."),
 		GetDefault<UDMMaterialValueMediaStream>()->GetComponentIcon(),
 		FUIAction(
 			FExecuteAction::CreateRaw(this,
@@ -196,11 +198,22 @@ void FDMMediaStreamStageSourceMenuExtender::AddMediaStreamLayerFromContext(UDMMe
 		return;
 	}
 
-	UDMMaterialLayerObject* OldLayer = InMenuContext->GetLayer();
+	EDMMaterialPropertyType PropertyType = EDMMaterialPropertyType::None;
 
-	if (!OldLayer)
+	if (UDMMaterialLayerObject* Layer = InMenuContext->GetLayer())
 	{
-		return;
+		PropertyType = Layer->GetMaterialProperty();
+	}
+	else
+	{
+		const TArray<TObjectPtr<UDMMaterialLayerObject>>& Layers = Slot->GetLayers();
+
+		if (Layers.IsEmpty())
+		{
+			return;
+		}
+
+		PropertyType = Layers.Last()->GetMaterialProperty();
 	}
 
 	if (GUndo)
@@ -212,7 +225,7 @@ void FDMMediaStreamStageSourceMenuExtender::AddMediaStreamLayerFromContext(UDMMe
 
 	const FDMUpdateGuard Guard;
 
-	NewLayer = Slot->AddDefaultLayer(OldLayer->GetMaterialProperty());
+	NewLayer = Slot->AddDefaultLayer(PropertyType);
 
 	if (!ensure(NewLayer))
 	{
diff --git a/Engine/Plugins/Experimental/MediaStream/Source/MediaStream/Private/MediaStreamSchemeHandlerLibrary.cpp b/Engine/Plugins/Experimental/MediaStream/Source/MediaStream/Private/MediaStreamSchemeHandlerLibrary.cpp
--- a/Engine/Plugins/Experimental/MediaStream/Source/MediaStream/Private/MediaStreamSchemeHandlerLibrary.cpp
+++ b/Engine/Plugins/Experimental/MediaStream/Source/MediaStream/Private/MediaStreamSchemeHandlerLibrary.cpp
@@ -1,13 +1,14 @@
 // Copyright Epic Games, Inc. All Rights Reserved.
 
 #include "IMediaStreamObjectHandler.h"
-#if WITH_EDITOR
 
 #include "IMediaStreamSchemeHandler.h"
 
+#if WITH_EDITOR
 #include "MediaStreamSchemeHandlerManager.h"
 #include "MediaStreamSource.h"
 #include "PropertyHandle.h"
+#endif
 
 FMediaStreamObjectHandlerCreatePlayerParams FMediaStreamSchemeHandlerCreatePlayerParams::operator<<(UObject* InMediaSource) const
 {
@@ -20,6 +21,7 @@ FMediaStreamObjectHandlerCreatePlayerParams FMediaStreamSchemeHandlerCreatePlaye
 	return NewParams;
 }
 
+#if WITH_EDITOR
 FMediaStreamSource* FMediaStreamSchemeHandlerLibrary::GetStreamSourcePtr(TWeakPtr<IPropertyHandle> InPropertyHandleWeak)
 {
 	TSharedPtr<IPropertyHandle> PropertyHandle = InPropertyHandleWeak.Pin();
