Material Designer: Apply button

diff --git a/Engine/Plugins/Experimental/DynamicMaterial/Source/DynamicMaterialEditor/Public/UI/Utils/IDMWidgetLibrary.h b/Engine/Plugins/Experimental/DynamicMaterial/Source/DynamicMaterialEditor/Public/UI/Utils/IDMWidgetLibrary.h
--- a/Engine/Plugins/Experimental/DynamicMaterial/Source/DynamicMaterialEditor/Public/UI/Utils/IDMWidgetLibrary.h
+++ b/Engine/Plugins/Experimental/DynamicMaterial/Source/DynamicMaterialEditor/Public/UI/Utils/IDMWidgetLibrary.h
@@ -4,19 +4,32 @@
 
 #include "HAL/Platform.h"
 
-class FName;
+#include "UObject/NameTypes.h"
+
+class FNotifyHook;
 class SDMMaterialComponentEditor;
 class SWidget;
+class UDynamicMaterialModelBase;
 class UObject;
 struct FDMPropertyHandle;
 
+struct FDMPropertyHandleGenerateParams
+{
+	const SWidget* Widget;
+	FNotifyHook* NotifyHook;
+	UDynamicMaterialModelBase* PreviewMaterialModelBase;
+	UDynamicMaterialModelBase* OriginalMaterialModelBase;
+	UObject* Object;
+	FName PropertyName;
+};
+
 class IDMWidgetLibrary
 {
 public:
+	virtual ~IDMWidgetLibrary() = default;
+
 	DYNAMICMATERIALEDITOR_API static IDMWidgetLibrary& Get();
 
 	// Creates a DM property handle for use in the details panel section of the Material Designer
-	virtual FDMPropertyHandle GetPropertyHandle(const SWidget* InOwningWidget, UObject* InObject, FName InPropertyName) = 0;
-
-	virtual FDMPropertyHandle GetPropertyHandle(const SDMMaterialComponentEditor* InOwningWidget, UObject* InObject, FName InPropertyName) = 0;
+	virtual FDMPropertyHandle GetPropertyHandle(const FDMPropertyHandleGenerateParams& InParams) = 0;
 };
diff --git a/Engine/Plugins/Experimental/DynamicMaterial/Source/DynamicMaterialEditor/Public/Utils/DMMaterialModelFunctionLibrary.h b/Engine/Plugins/Experimental/DynamicMaterial/Source/DynamicMaterialEditor/Public/Utils/DMMaterialModelFunctionLibrary.h
--- a/Engine/Plugins/Experimental/DynamicMaterial/Source/DynamicMaterialEditor/Public/Utils/DMMaterialModelFunctionLibrary.h
+++ b/Engine/Plugins/Experimental/DynamicMaterial/Source/DynamicMaterialEditor/Public/Utils/DMMaterialModelFunctionLibrary.h
@@ -45,4 +45,12 @@ public:
 	DYNAMICMATERIALEDITOR_API static bool CreateModelInstanceInMaterial(UDynamicMaterialModel* InFromModel, UDynamicMaterialInstance* InToInstance);
 
 	DYNAMICMATERIALEDITOR_API static FString RemoveAssetPrefix(const FString& InAssetName);
+
+	DYNAMICMATERIALEDITOR_API static UDynamicMaterialModelBase* CreatePreviewModel(UDynamicMaterialModelBase* InOriginalModelBase);
+
+	/**
+	 * Ensures that the source model is entirely mirrored to the target model by updating
+	 * objects, rather than recreating them. Additions and removals are also propagated.
+	 */
+	DYNAMICMATERIALEDITOR_API static void MirrorMaterialModel(UDynamicMaterialModelBase* InSource, UDynamicMaterialModelBase*& InTarget);
 };
diff --git a/Engine/Plugins/Experimental/DynamicMaterialMediaStreamBridge/Source/DynamicMaterialMediaStreamBridge/Private/DMMaterialValueMediaStream.cpp b/Engine/Plugins/Experimental/DynamicMaterialMediaStreamBridge/Source/DynamicMaterialMediaStreamBridge/Private/DMMaterialValueMediaStream.cpp
--- a/Engine/Plugins/Experimental/DynamicMaterialMediaStreamBridge/Source/DynamicMaterialMediaStreamBridge/Private/DMMaterialValueMediaStream.cpp
+++ b/Engine/Plugins/Experimental/DynamicMaterialMediaStreamBridge/Source/DynamicMaterialMediaStreamBridge/Private/DMMaterialValueMediaStream.cpp
@@ -215,6 +215,11 @@ void UDMMaterialValueMediaStream::UpdatePlayer()
 	UpdateTextureFromMediaStream();
 }
 
+void UDMMaterialValueMediaStream::CopyParametersFrom_Implementation(UObject* InOther)
+{
+	// Overridden to do nothing
+}
+
 void UDMMaterialValueMediaStream::OnComponentAdded()
 {
 	Super::OnComponentAdded();
diff --git a/Engine/Plugins/Experimental/DynamicMaterialMediaStreamBridge/Source/DynamicMaterialMediaStreamBridge/Private/DMMaterialValueMediaStreamDynamic.cpp b/Engine/Plugins/Experimental/DynamicMaterialMediaStreamBridge/Source/DynamicMaterialMediaStreamBridge/Private/DMMaterialValueMediaStreamDynamic.cpp
--- a/Engine/Plugins/Experimental/DynamicMaterialMediaStreamBridge/Source/DynamicMaterialMediaStreamBridge/Private/DMMaterialValueMediaStreamDynamic.cpp
+++ b/Engine/Plugins/Experimental/DynamicMaterialMediaStreamBridge/Source/DynamicMaterialMediaStreamBridge/Private/DMMaterialValueMediaStreamDynamic.cpp
@@ -206,6 +206,11 @@ void UDMMaterialValueMediaStreamDynamic::OnComponentRemoved()
 
 	SetValue(nullptr);
 }
+
+void UDMMaterialValueMediaStreamDynamic::CopyParametersFrom_Implementation(UObject* InOther)
+{
+	// Overridden to do nothing
+}
 #endif
 
 #undef LOCTEXT_NAMESPACE
diff --git a/Engine/Plugins/Experimental/DynamicMaterialMediaStreamBridge/Source/DynamicMaterialMediaStreamBridge/Public/DMMaterialValueMediaStream.h b/Engine/Plugins/Experimental/DynamicMaterialMediaStreamBridge/Source/DynamicMaterialMediaStreamBridge/Public/DMMaterialValueMediaStream.h
--- a/Engine/Plugins/Experimental/DynamicMaterialMediaStreamBridge/Source/DynamicMaterialMediaStreamBridge/Public/DMMaterialValueMediaStream.h
+++ b/Engine/Plugins/Experimental/DynamicMaterialMediaStreamBridge/Source/DynamicMaterialMediaStreamBridge/Public/DMMaterialValueMediaStream.h
@@ -67,6 +67,10 @@ protected:
 
 	void UpdatePlayer();
 
+	//~ Begin IDMParameterContainer
+	virtual void CopyParametersFrom_Implementation(UObject* InOther) override;
+	//~ End IDMParameterContainer
+
 	//~ Begin UDMMaterialComponent
 	DYNAMICMATERIALMEDIASTREAMBRIDGE_API  virtual void OnComponentAdded() override;
 	DYNAMICMATERIALMEDIASTREAMBRIDGE_API  virtual void OnComponentRemoved() override;
diff --git a/Engine/Plugins/Experimental/DynamicMaterialMediaStreamBridge/Source/DynamicMaterialMediaStreamBridgeEditor/Private/DMMaterialValueMediaStreamPropertyRowGenerator.cpp b/Engine/Plugins/Experimental/DynamicMaterialMediaStreamBridge/Source/DynamicMaterialMediaStreamBridgeEditor/Private/DMMaterialValueMediaStreamPropertyRowGenerator.cpp
--- a/Engine/Plugins/Experimental/DynamicMaterialMediaStreamBridge/Source/DynamicMaterialMediaStreamBridgeEditor/Private/DMMaterialValueMediaStreamPropertyRowGenerator.cpp
+++ b/Engine/Plugins/Experimental/DynamicMaterialMediaStreamBridge/Source/DynamicMaterialMediaStreamBridgeEditor/Private/DMMaterialValueMediaStreamPropertyRowGenerator.cpp
@@ -5,6 +5,7 @@
 #include "Components/DMMaterialComponent.h"
 #include "DMEDefs.h"
 #include "DMMaterialValueMediaStream.h"
+#include "IDynamicMaterialEditorModule.h"
 #include "IMediaStreamPlayer.h"
 #include "MediaStream.h"
 #include "UI/Utils/IDMWidgetLibrary.h"
@@ -16,27 +17,26 @@ const TSharedRef<FDMMaterialValueMediaStreamPropertyRowGenerator>& FDMMaterialVa
 	return Generator;
 }
 
-void FDMMaterialValueMediaStreamPropertyRowGenerator::AddComponentProperties(const TSharedRef<SDMMaterialComponentEditor>& InComponentEditorWidget, UDMMaterialComponent* InComponent,
-	TArray<FDMPropertyHandle>& InOutPropertyRows, TSet<UDMMaterialComponent*>& InOutProcessedObjects)
+void FDMMaterialValueMediaStreamPropertyRowGenerator::AddComponentProperties(FDMComponentPropertyRowGeneratorParams& InParams)
 {
-	if (!IsValid(InComponent))
+	if (!IsValid(InParams.Object))
 	{
 		return;
 	}
 
-	if (InOutProcessedObjects.Contains(InComponent))
+	if (InParams.ProcessedObjects.Contains(InParams.Object))
 	{
 		return;
 	}
 
-	UDMMaterialValueMediaStream* MediaStreamValue = Cast<UDMMaterialValueMediaStream>(InComponent);
+	UDMMaterialValueMediaStream* MediaStreamValue = Cast<UDMMaterialValueMediaStream>(InParams.Object);
 
 	if (!MediaStreamValue)
 	{
 		return;
 	}
 
-	InOutProcessedObjects.Add(InComponent);
+	InParams.ProcessedObjects.Add(InParams.Object);
 
 	UMediaStream* MediaStream = MediaStreamValue->GetMediaStream();
 
@@ -45,45 +45,50 @@ void FDMMaterialValueMediaStreamPropertyRowGenerator::AddComponentProperties(con
 		return;
 	}
 
-	AddSourceProperty(InComponentEditorWidget, InOutPropertyRows, MediaStream);
-	AddPlayerProperty(InComponentEditorWidget, InOutPropertyRows, MediaStream);
+	FDMComponentPropertyRowGeneratorParams MediaStreamParams = InParams;
+	MediaStreamParams.Object = MediaStream;
+
+	AddSourceProperty(MediaStreamParams);
+	AddPlayerProperty(MediaStreamParams);
 }
 
-void FDMMaterialValueMediaStreamPropertyRowGenerator::AddSourceProperty(const TSharedRef<SDMMaterialComponentEditor>& InComponentEditorWidget, 
-	TArray<FDMPropertyHandle>& InOutPropertyRows, UMediaStream* InMediaStream)
+void FDMMaterialValueMediaStreamPropertyRowGenerator::AddSourceProperty(FDMComponentPropertyRowGeneratorParams& InParams)
 {
-	FDMPropertyHandle& NewHandle = InOutPropertyRows.Add_GetRef(
-		IDMWidgetLibrary::Get().GetPropertyHandle(&*InComponentEditorWidget, InMediaStream, TEXT("Source"))
+	FDMPropertyHandle& NewHandle = InParams.PropertyRows.Add_GetRef(
+		IDMWidgetLibrary::Get().GetPropertyHandle(InParams.CreatePropertyHandleParams(TEXT("Source")))
 	);
 }
 
-void FDMMaterialValueMediaStreamPropertyRowGenerator::AddPlayerProperty(const TSharedRef<SDMMaterialComponentEditor>& InComponentEditorWidget,
-	TArray<FDMPropertyHandle>& InOutPropertyRows, UMediaStream* InMediaStream)
+void FDMMaterialValueMediaStreamPropertyRowGenerator::AddPlayerProperty(FDMComponentPropertyRowGeneratorParams& InParams)
 {
-	IMediaStreamPlayer* Player = InMediaStream->GetPlayer().GetInterface();
+	UMediaStream* MediaStream = Cast<UMediaStream>(InParams.Object);
+
+	IMediaStreamPlayer* Player = MediaStream->GetPlayer().GetInterface();
 
 	if (!Player || Player->IsReadOnly())
 	{
 		return;
 	}
 
-	FDMPropertyHandle& ControlHandle = InOutPropertyRows.Add_GetRef(
-		IDMWidgetLibrary::Get().GetPropertyHandle(&*InComponentEditorWidget, InMediaStream, TEXT("Control"))
+	const FName VideoPlayerCategory = TEXT("Video Player");
+
+	FDMPropertyHandle& ControlHandle = InParams.PropertyRows.Add_GetRef(
+		IDMWidgetLibrary::Get().GetPropertyHandle(InParams.CreatePropertyHandleParams(TEXT("Control")))
 	);
-	ControlHandle.CategoryOverrideName = TEXT("Video Player");
+	ControlHandle.CategoryOverrideName = VideoPlayerCategory;
 
-	FDMPropertyHandle& TextureHandle = InOutPropertyRows.Add_GetRef(
-		IDMWidgetLibrary::Get().GetPropertyHandle(&*InComponentEditorWidget, InMediaStream, TEXT("MediaTexture"))
+	FDMPropertyHandle& TextureHandle = InParams.PropertyRows.Add_GetRef(
+		IDMWidgetLibrary::Get().GetPropertyHandle(InParams.CreatePropertyHandleParams(TEXT("MediaTexture")))
 	);
-	TextureHandle.CategoryOverrideName = TEXT("Video Player");
+	TextureHandle.CategoryOverrideName = VideoPlayerCategory;
 
-	FDMPropertyHandle& TextureConfigHandle = InOutPropertyRows.Add_GetRef(
-		IDMWidgetLibrary::Get().GetPropertyHandle(&*InComponentEditorWidget, InMediaStream, TEXT("TextureConfig"))
+	FDMPropertyHandle& TextureConfigHandle = InParams.PropertyRows.Add_GetRef(
+		IDMWidgetLibrary::Get().GetPropertyHandle(InParams.CreatePropertyHandleParams(TEXT("TextureConfig")))
 	);
-	TextureConfigHandle.CategoryOverrideName = TEXT("Video Player");
+	TextureConfigHandle.CategoryOverrideName = VideoPlayerCategory;
 
-	FDMPropertyHandle& PlayerConfigHandle = InOutPropertyRows.Add_GetRef(
-		IDMWidgetLibrary::Get().GetPropertyHandle(&*InComponentEditorWidget, InMediaStream, TEXT("PlayerConfig"))
+	FDMPropertyHandle& PlayerConfigHandle = InParams.PropertyRows.Add_GetRef(
+		IDMWidgetLibrary::Get().GetPropertyHandle(InParams.CreatePropertyHandleParams(TEXT("PlayerConfig")))
 	);
-	PlayerConfigHandle.CategoryOverrideName = TEXT("Video Player");
+	PlayerConfigHandle.CategoryOverrideName = VideoPlayerCategory;
 }
diff --git a/Engine/Plugins/Experimental/DynamicMaterialMediaStreamBridge/Source/DynamicMaterialMediaStreamBridgeEditor/Private/DMMaterialValueMediaStreamPropertyRowGenerator.h b/Engine/Plugins/Experimental/DynamicMaterialMediaStreamBridge/Source/DynamicMaterialMediaStreamBridgeEditor/Private/DMMaterialValueMediaStreamPropertyRowGenerator.h
--- a/Engine/Plugins/Experimental/DynamicMaterialMediaStreamBridge/Source/DynamicMaterialMediaStreamBridgeEditor/Private/DMMaterialValueMediaStreamPropertyRowGenerator.h
+++ b/Engine/Plugins/Experimental/DynamicMaterialMediaStreamBridge/Source/DynamicMaterialMediaStreamBridgeEditor/Private/DMMaterialValueMediaStreamPropertyRowGenerator.h
@@ -14,13 +14,10 @@ public:
 	virtual ~FDMMaterialValueMediaStreamPropertyRowGenerator() override = default;
 
 	//~ Begin FDMComponentPropertyRowGenerator
-	virtual void AddComponentProperties(const TSharedRef<SDMMaterialComponentEditor>& InComponentEditorWidget, UDMMaterialComponent* InComponent,
-		TArray<FDMPropertyHandle>& InOutPropertyRows, TSet<UDMMaterialComponent*>& InOutProcessedObjects) override;
+	virtual void AddComponentProperties(FDMComponentPropertyRowGeneratorParams& InParams) override;
 	//~ End FDMComponentPropertyRowGenerator
 
-	void AddSourceProperty(const TSharedRef<SDMMaterialComponentEditor>& InComponentEditorWidget, TArray<FDMPropertyHandle>& InOutPropertyRows, 
-		UMediaStream* InMediaStream);
+	void AddSourceProperty(FDMComponentPropertyRowGeneratorParams& InParams);
 
-	void AddPlayerProperty(const TSharedRef<SDMMaterialComponentEditor>& InComponentEditorWidget, TArray<FDMPropertyHandle>& InOutPropertyRows,
-		UMediaStream* InMediaStream);
+	void AddPlayerProperty(FDMComponentPropertyRowGeneratorParams& InParams);
 };
