Material Designer: Fixed material effect property enumeration

diff --git a/Engine/Plugins/Experimental/DynamicMaterial/Source/DynamicMaterialEditor/Private/UI/PropertyGenerators/DMComponentPropertyRowGenerator.cpp b/Engine/Plugins/Experimental/DynamicMaterial/Source/DynamicMaterialEditor/Private/UI/PropertyGenerators/DMComponentPropertyRowGenerator.cpp
--- a/Engine/Plugins/Experimental/DynamicMaterial/Source/DynamicMaterialEditor/Private/UI/PropertyGenerators/DMComponentPropertyRowGenerator.cpp
+++ b/Engine/Plugins/Experimental/DynamicMaterial/Source/DynamicMaterialEditor/Private/UI/PropertyGenerators/DMComponentPropertyRowGenerator.cpp
@@ -106,7 +106,7 @@ void FDMComponentPropertyRowGenerator::AddPropertyEditRows(FDMComponentPropertyR
 		}
 	}
 
-	FDMPropertyHandle& Handle = InParams.PropertyRows.Add_GetRef(FDMWidgetLibrary::Get().GetPropertyHandle(InParams.CreatePropertyHandleParams(InProperty->GetFName())));
+	FDMPropertyHandle& Handle = InParams.PropertyRows->Add_GetRef(FDMWidgetLibrary::Get().GetPropertyHandle(InParams.CreatePropertyHandleParams(InProperty->GetFName())));
 	Handle.bEnabled = !IsDynamic(InParams);
 }
 
diff --git a/Engine/Plugins/Experimental/DynamicMaterial/Source/DynamicMaterialEditor/Private/UI/PropertyGenerators/DMMaterialEffectFunctionPropertyRowGenerator.cpp b/Engine/Plugins/Experimental/DynamicMaterial/Source/DynamicMaterialEditor/Private/UI/PropertyGenerators/DMMaterialEffectFunctionPropertyRowGenerator.cpp
--- a/Engine/Plugins/Experimental/DynamicMaterial/Source/DynamicMaterialEditor/Private/UI/PropertyGenerators/DMMaterialEffectFunctionPropertyRowGenerator.cpp
+++ b/Engine/Plugins/Experimental/DynamicMaterial/Source/DynamicMaterialEditor/Private/UI/PropertyGenerators/DMMaterialEffectFunctionPropertyRowGenerator.cpp
@@ -76,7 +76,7 @@ void FDMMaterialEffectFunctionPropertyRowGenerator::AddComponentProperties(FDMCo
 
 		FDMComponentPropertyRowGeneratorParams ValueParams = InParams;
 		ValueParams.Object = Value;
-		ValueParams.PropertyRows = ValuePropertyRows;
+		ValueParams.PropertyRows = &ValuePropertyRows;
 
 		FDynamicMaterialEditorModule::GeneratorComponentPropertyRows(ValueParams);
 
@@ -119,7 +119,7 @@ void FDMMaterialEffectFunctionPropertyRowGenerator::AddComponentProperties(FDMCo
 
 		}
 
-		InParams.PropertyRows.Append(ValuePropertyRows);
+		InParams.PropertyRows->Append(ValuePropertyRows);
 	}
 }
 
diff --git a/Engine/Plugins/Experimental/DynamicMaterial/Source/DynamicMaterialEditor/Private/UI/PropertyGenerators/DMMaterialModelPropertyRowGenerator.cpp b/Engine/Plugins/Experimental/DynamicMaterial/Source/DynamicMaterialEditor/Private/UI/PropertyGenerators/DMMaterialModelPropertyRowGenerator.cpp
--- a/Engine/Plugins/Experimental/DynamicMaterial/Source/DynamicMaterialEditor/Private/UI/PropertyGenerators/DMMaterialModelPropertyRowGenerator.cpp
+++ b/Engine/Plugins/Experimental/DynamicMaterial/Source/DynamicMaterialEditor/Private/UI/PropertyGenerators/DMMaterialModelPropertyRowGenerator.cpp
@@ -90,7 +90,7 @@ void FDMMaterialModelPropertyRowGenerator::AddGlobalValue(FDMComponentPropertyRo
 	FDMComponentPropertyRowGeneratorParams GlobalValueParams = InParams;
 	GlobalValueParams.Object = InComponent;
 
-	FDMPropertyHandle& ComponentHandle = InParams.PropertyRows.Add_GetRef(FDMWidgetLibrary::Get().GetPropertyHandle(
+	FDMPropertyHandle& ComponentHandle = InParams.PropertyRows->Add_GetRef(FDMWidgetLibrary::Get().GetPropertyHandle(
 		GlobalValueParams.CreatePropertyHandleParams(UDMMaterialValue::ValueName)
 	));
 
@@ -120,7 +120,7 @@ void FDMMaterialModelPropertyRowGenerator::AddVariable(FDMComponentPropertyRowGe
 	FDMComponentPropertyRowGeneratorParams VariableParams = InParams;
 	VariableParams.Object = InObject;
 
-	FDMPropertyHandle& ValueHandle = InParams.PropertyRows.Add_GetRef(FDMWidgetLibrary::Get().GetPropertyHandle(
+	FDMPropertyHandle& ValueHandle = InParams.PropertyRows->Add_GetRef(FDMWidgetLibrary::Get().GetPropertyHandle(
 		VariableParams.CreatePropertyHandleParams(InPropertyName)
 	));
 
diff --git a/Engine/Plugins/Experimental/DynamicMaterial/Source/DynamicMaterialEditor/Private/UI/PropertyGenerators/DMMaterialStageFunctionPropertyRowGenerator.cpp b/Engine/Plugins/Experimental/DynamicMaterial/Source/DynamicMaterialEditor/Private/UI/PropertyGenerators/DMMaterialStageFunctionPropertyRowGenerator.cpp
--- a/Engine/Plugins/Experimental/DynamicMaterial/Source/DynamicMaterialEditor/Private/UI/PropertyGenerators/DMMaterialStageFunctionPropertyRowGenerator.cpp
+++ b/Engine/Plugins/Experimental/DynamicMaterial/Source/DynamicMaterialEditor/Private/UI/PropertyGenerators/DMMaterialStageFunctionPropertyRowGenerator.cpp
@@ -83,7 +83,7 @@ void FDMMaterialStageFunctionPropertyRowGenerator::AddComponentProperties(FDMCom
 
 			FDMComponentPropertyRowGeneratorParams ValueParams = InParams;
 			ValueParams.Object = Value;
-			ValueParams.PropertyRows = ValuePropertyRows;
+			ValueParams.PropertyRows = &ValuePropertyRows;
 
 			FDynamicMaterialEditorModule::GeneratorComponentPropertyRows(ValueParams);
 
@@ -142,7 +142,7 @@ void FDMMaterialStageFunctionPropertyRowGenerator::AddComponentProperties(FDMCom
 
 	FDynamicMaterialEditorModule::GeneratorComponentPropertyRows(FunctionParams);
 
-	InParams.PropertyRows.Append(AllValuePropertyRows);
+	InParams.PropertyRows->Append(AllValuePropertyRows);
 }
 
 #undef LOCTEXT_NAMESPACE
diff --git a/Engine/Plugins/Experimental/DynamicMaterial/Source/DynamicMaterialEditor/Private/UI/PropertyGenerators/DMMaterialValueDynamicPropertyRowGenerator.cpp b/Engine/Plugins/Experimental/DynamicMaterial/Source/DynamicMaterialEditor/Private/UI/PropertyGenerators/DMMaterialValueDynamicPropertyRowGenerator.cpp
--- a/Engine/Plugins/Experimental/DynamicMaterial/Source/DynamicMaterialEditor/Private/UI/PropertyGenerators/DMMaterialValueDynamicPropertyRowGenerator.cpp
+++ b/Engine/Plugins/Experimental/DynamicMaterial/Source/DynamicMaterialEditor/Private/UI/PropertyGenerators/DMMaterialValueDynamicPropertyRowGenerator.cpp
@@ -63,11 +63,11 @@ void FDMMaterialValueDynamicPropertyRowGenerator::AddComponentProperties(FDMComp
 
 		Handle.bEnabled = true;
 
-		InParams.PropertyRows.Add(Handle);
+		InParams.PropertyRows->Add(Handle);
 	}
 
 	const TArray<FName>& Properties = ParentValue->GetEditableProperties();
-	const int32 StartRow = InParams.PropertyRows.Num();
+	const int32 StartRow = InParams.PropertyRows->Num();
 
 	for (const FName& Property : Properties)
 	{
@@ -82,10 +82,10 @@ void FDMMaterialValueDynamicPropertyRowGenerator::AddComponentProperties(FDMComp
 		}
 	}
 
-	const int32 EndRow = InParams.PropertyRows.Num();
+	const int32 EndRow = InParams.PropertyRows->Num();
 
 	for (int32 RowIndex = StartRow; RowIndex < EndRow; ++RowIndex)
 	{
-		InParams.PropertyRows[RowIndex].bEnabled = false;
+		(*InParams.PropertyRows)[RowIndex].bEnabled = false;
 	}
 }
diff --git a/Engine/Plugins/Experimental/DynamicMaterial/Source/DynamicMaterialEditor/Private/UI/PropertyGenerators/DMMaterialValuePropertyRowGenerator.cpp b/Engine/Plugins/Experimental/DynamicMaterial/Source/DynamicMaterialEditor/Private/UI/PropertyGenerators/DMMaterialValuePropertyRowGenerator.cpp
--- a/Engine/Plugins/Experimental/DynamicMaterial/Source/DynamicMaterialEditor/Private/UI/PropertyGenerators/DMMaterialValuePropertyRowGenerator.cpp
+++ b/Engine/Plugins/Experimental/DynamicMaterial/Source/DynamicMaterialEditor/Private/UI/PropertyGenerators/DMMaterialValuePropertyRowGenerator.cpp
@@ -76,7 +76,7 @@ void FDMMaterialValuePropertyRowGenerator::AddComponentProperties(FDMComponentPr
 
 		Handle.bEnabled = true;
 
-		InParams.PropertyRows.Add(Handle);
+		InParams.PropertyRows->Add(Handle);
 	}
 
 	const TArray<FName>& Properties = Value->GetEditableProperties();
diff --git a/Engine/Plugins/Experimental/DynamicMaterial/Source/DynamicMaterialEditor/Private/UI/PropertyGenerators/DMTextureUVDynamicPropertyRowGenerator.cpp b/Engine/Plugins/Experimental/DynamicMaterial/Source/DynamicMaterialEditor/Private/UI/PropertyGenerators/DMTextureUVDynamicPropertyRowGenerator.cpp
--- a/Engine/Plugins/Experimental/DynamicMaterial/Source/DynamicMaterialEditor/Private/UI/PropertyGenerators/DMTextureUVDynamicPropertyRowGenerator.cpp
+++ b/Engine/Plugins/Experimental/DynamicMaterial/Source/DynamicMaterialEditor/Private/UI/PropertyGenerators/DMTextureUVDynamicPropertyRowGenerator.cpp
@@ -128,7 +128,7 @@ void UE::DynamicMaterialEditor::Private::AddTextureUVDynamicPropertyRow(FDMCompo
 	FDMComponentPropertyRowGeneratorParams ComponentParams = InParams;
 	ComponentParams.Object = InComponent;
 
-	FDMPropertyHandle& NewHandle = InParams.PropertyRows.Add_GetRef(FDMWidgetLibrary::Get().GetPropertyHandle(ComponentParams.CreatePropertyHandleParams(InProperty)));
+	FDMPropertyHandle& NewHandle = InParams.PropertyRows->Add_GetRef(FDMWidgetLibrary::Get().GetPropertyHandle(ComponentParams.CreatePropertyHandleParams(InProperty)));
 
 	NewHandle.ResetToDefaultOverride = FResetToDefaultOverride::Create(
 		FIsResetToDefaultVisible::CreateStatic(&UE::DynamicMaterialEditor::Private::CanResetTextureUVDynamicDynamicPropertyToDefault),
@@ -170,7 +170,7 @@ void UE::DynamicMaterialEditor::Private::AddTextureUVDynamicVisualizerRow(FDMCom
 	VisualizerHandle.ValueWidget = SNew(SDMTextureUVVisualizerProperty, EditorWidget.ToSharedRef(), Stage).TextureUVDynamic(InTextureUVDynamic);
 	VisualizerHandle.CategoryOverrideName = TEXT("Texture UV");
 	VisualizerHandle.bEnabled = true;
-	InParams.PropertyRows.Add(VisualizerHandle);
+	InParams.PropertyRows->Add(VisualizerHandle);
 }
 
 
diff --git a/Engine/Plugins/Experimental/DynamicMaterial/Source/DynamicMaterialEditor/Private/UI/PropertyGenerators/DMTextureUVPropertyRowGenerator.cpp b/Engine/Plugins/Experimental/DynamicMaterial/Source/DynamicMaterialEditor/Private/UI/PropertyGenerators/DMTextureUVPropertyRowGenerator.cpp
--- a/Engine/Plugins/Experimental/DynamicMaterial/Source/DynamicMaterialEditor/Private/UI/PropertyGenerators/DMTextureUVPropertyRowGenerator.cpp
+++ b/Engine/Plugins/Experimental/DynamicMaterial/Source/DynamicMaterialEditor/Private/UI/PropertyGenerators/DMTextureUVPropertyRowGenerator.cpp
@@ -128,7 +128,7 @@ bool FDMTextureUVPropertyRowGenerator::AllowKeyframeButton(UDMMaterialComponent*
 
 void UE::DynamicMaterialEditor::Private::AddTextureUVPropertyRow(FDMComponentPropertyRowGeneratorParams& InParams, FName InProperty)
 {
-	FDMPropertyHandle& NewHandle = InParams.PropertyRows.Add_GetRef(FDMWidgetLibrary::Get().GetPropertyHandle(InParams.CreatePropertyHandleParams(InProperty)));
+	FDMPropertyHandle& NewHandle = InParams.PropertyRows->Add_GetRef(FDMWidgetLibrary::Get().GetPropertyHandle(InParams.CreatePropertyHandleParams(InProperty)));
 
 	NewHandle.ResetToDefaultOverride = FResetToDefaultOverride::Create(
 		FIsResetToDefaultVisible::CreateStatic(&UE::DynamicMaterialEditor::Private::CanResetTextureUVPropertyToDefault),
@@ -164,7 +164,7 @@ void UE::DynamicMaterialEditor::Private::AddTextureUVVisualizerRow(FDMComponentP
 	VisualizerHandle.ValueWidget = SNew(SDMTextureUVVisualizerProperty, EditorWidget.ToSharedRef(), Stage).TextureUV(TextureUV);
 	VisualizerHandle.CategoryOverrideName = TEXT("Texture UV");
 	VisualizerHandle.bEnabled = true;
-	InParams.PropertyRows.Add(VisualizerHandle);
+	InParams.PropertyRows->Add(VisualizerHandle);
 }
 
 
diff --git a/Engine/Plugins/Experimental/DynamicMaterial/Source/DynamicMaterialEditor/Private/UI/PropertyGenerators/DMThroughputPropertyRowGenerator.cpp b/Engine/Plugins/Experimental/DynamicMaterial/Source/DynamicMaterialEditor/Private/UI/PropertyGenerators/DMThroughputPropertyRowGenerator.cpp
--- a/Engine/Plugins/Experimental/DynamicMaterial/Source/DynamicMaterialEditor/Private/UI/PropertyGenerators/DMThroughputPropertyRowGenerator.cpp
+++ b/Engine/Plugins/Experimental/DynamicMaterial/Source/DynamicMaterialEditor/Private/UI/PropertyGenerators/DMThroughputPropertyRowGenerator.cpp
@@ -68,7 +68,7 @@ void FDMThroughputPropertyRowGenerator::AddComponentProperties(FDMComponentPrope
 				continue;
 			}
 
-			const int32 StartRow = InParams.PropertyRows.Num();
+			const int32 StartRow = InParams.PropertyRows->Num();
 
 			for (const FDMMaterialStageConnectorChannel& Channel : InputMap[InputIdx].Channels)
 			{
@@ -83,9 +83,9 @@ void FDMThroughputPropertyRowGenerator::AddComponentProperties(FDMComponentPrope
 				}
 			}
 
-			for (int32 PropertyRowIdx = StartRow; PropertyRowIdx < InParams.PropertyRows.Num(); ++PropertyRowIdx)
+			for (int32 PropertyRowIdx = StartRow; PropertyRowIdx < InParams.PropertyRows->Num(); ++PropertyRowIdx)
 			{
-				FDMPropertyHandle& PropertyRow = InParams.PropertyRows[PropertyRowIdx];
+				FDMPropertyHandle& PropertyRow = (*InParams.PropertyRows)[PropertyRowIdx];
 
 				if (PropertyRow.NameOverride.IsSet())
 				{
diff --git a/Engine/Plugins/Experimental/DynamicMaterial/Source/DynamicMaterialEditor/Public/IDynamicMaterialEditorModule.h b/Engine/Plugins/Experimental/DynamicMaterial/Source/DynamicMaterialEditor/Public/IDynamicMaterialEditorModule.h
--- a/Engine/Plugins/Experimental/DynamicMaterial/Source/DynamicMaterialEditor/Public/IDynamicMaterialEditorModule.h
+++ b/Engine/Plugins/Experimental/DynamicMaterial/Source/DynamicMaterialEditor/Public/IDynamicMaterialEditorModule.h
@@ -6,6 +6,7 @@
 
 #include "Containers/ContainersFwd.h"
 #include "Delegates/Delegate.h"
+#include "Misc/NotNull.h"
 #include "Templates/SharedPointerFwd.h"
 #include "UI/Utils/IDMWidgetLibrary.h"
 
@@ -29,11 +30,11 @@ struct FDMComponentPropertyRowGeneratorParams
 	UDynamicMaterialModelBase* PreviewMaterialModelBase = nullptr;
 	UDynamicMaterialModelBase* OriginalMaterialModelBase = nullptr;
 	UObject* Object = nullptr;
-	TArray<FDMPropertyHandle>& PropertyRows;
+	TNotNull<TArray<FDMPropertyHandle>*> PropertyRows;
 	TSet<UObject*>& ProcessedObjects;
 
 	FDMComponentPropertyRowGeneratorParams(TArray<FDMPropertyHandle>& InPropertyRows, TSet<UObject*>& InProcessedObjects)
-		: PropertyRows(InPropertyRows)
+		: PropertyRows(&InPropertyRows)
 		, ProcessedObjects(InProcessedObjects)
 	{
 	}
diff --git a/Engine/Plugins/Experimental/DynamicMaterialMediaStreamBridge/Source/DynamicMaterialMediaStreamBridgeEditor/Private/DMMaterialValueMediaStreamPropertyRowGenerator.cpp b/Engine/Plugins/Experimental/DynamicMaterialMediaStreamBridge/Source/DynamicMaterialMediaStreamBridgeEditor/Private/DMMaterialValueMediaStreamPropertyRowGenerator.cpp
--- a/Engine/Plugins/Experimental/DynamicMaterialMediaStreamBridge/Source/DynamicMaterialMediaStreamBridgeEditor/Private/DMMaterialValueMediaStreamPropertyRowGenerator.cpp
+++ b/Engine/Plugins/Experimental/DynamicMaterialMediaStreamBridge/Source/DynamicMaterialMediaStreamBridgeEditor/Private/DMMaterialValueMediaStreamPropertyRowGenerator.cpp
@@ -54,7 +54,7 @@ void FDMMaterialValueMediaStreamPropertyRowGenerator::AddComponentProperties(FDM
 
 void FDMMaterialValueMediaStreamPropertyRowGenerator::AddSourceProperty(FDMComponentPropertyRowGeneratorParams& InParams)
 {
-	FDMPropertyHandle& NewHandle = InParams.PropertyRows.Add_GetRef(
+	FDMPropertyHandle& NewHandle = InParams.PropertyRows->Add_GetRef(
 		IDMWidgetLibrary::Get().GetPropertyHandle(InParams.CreatePropertyHandleParams(TEXT("Source")))
 	);
 }
@@ -72,22 +72,22 @@ void FDMMaterialValueMediaStreamPropertyRowGenerator::AddPlayerProperty(FDMCompo
 
 	const FName VideoPlayerCategory = TEXT("Video Player");
 
-	FDMPropertyHandle& ControlHandle = InParams.PropertyRows.Add_GetRef(
+	FDMPropertyHandle& ControlHandle = InParams.PropertyRows->Add_GetRef(
 		IDMWidgetLibrary::Get().GetPropertyHandle(InParams.CreatePropertyHandleParams(TEXT("Control")))
 	);
 	ControlHandle.CategoryOverrideName = VideoPlayerCategory;
 
-	FDMPropertyHandle& TextureHandle = InParams.PropertyRows.Add_GetRef(
+	FDMPropertyHandle& TextureHandle = InParams.PropertyRows->Add_GetRef(
 		IDMWidgetLibrary::Get().GetPropertyHandle(InParams.CreatePropertyHandleParams(TEXT("MediaTexture")))
 	);
 	TextureHandle.CategoryOverrideName = VideoPlayerCategory;
 
-	FDMPropertyHandle& TextureConfigHandle = InParams.PropertyRows.Add_GetRef(
+	FDMPropertyHandle& TextureConfigHandle = InParams.PropertyRows->Add_GetRef(
 		IDMWidgetLibrary::Get().GetPropertyHandle(InParams.CreatePropertyHandleParams(TEXT("TextureConfig")))
 	);
 	TextureConfigHandle.CategoryOverrideName = VideoPlayerCategory;
 
-	FDMPropertyHandle& PlayerConfigHandle = InParams.PropertyRows.Add_GetRef(
+	FDMPropertyHandle& PlayerConfigHandle = InParams.PropertyRows->Add_GetRef(
 		IDMWidgetLibrary::Get().GetPropertyHandle(InParams.CreatePropertyHandleParams(TEXT("PlayerConfig")))
 	);
 	PlayerConfigHandle.CategoryOverrideName = VideoPlayerCategory;
