Material Designer: Fixed media source
- Media sources now properly work between the preview and source materials.
- Added play controls for both the preview and source materials.
- Adding a track now correctly adds the track for the source material version.
- Added playback state, requested seek frame and playlist index properties.

diff --git a/Engine/Plugins/Experimental/DynamicMaterial/Source/DynamicMaterialEditor/Private/Components/MaterialStageInputs/DMMSIValue.cpp b/Engine/Plugins/Experimental/DynamicMaterial/Source/DynamicMaterialEditor/Private/Components/MaterialStageInputs/DMMSIValue.cpp
--- a/Engine/Plugins/Experimental/DynamicMaterial/Source/DynamicMaterialEditor/Private/Components/MaterialStageInputs/DMMSIValue.cpp
+++ b/Engine/Plugins/Experimental/DynamicMaterial/Source/DynamicMaterialEditor/Private/Components/MaterialStageInputs/DMMSIValue.cpp
@@ -558,8 +558,12 @@ void UDMMaterialStageInputValue::ApplyDefaultLayerSettings()
 			if (UTexture* Texture = DefaultSlotValue.Texture.LoadSynchronous())
 			{
 				UDMMaterialValueTexture* TextureValue = Cast<UDMMaterialValueTexture>(Value);
-				TextureValue->SetDefaultValue(Texture);
-				TextureValue->ApplyDefaultValue();
+
+				if (TextureValue->GetClass() == UDMMaterialValueTexture::StaticClass() || !TextureValue->GetDefaultValue())
+				{
+					TextureValue->SetDefaultValue(Texture);
+					TextureValue->ApplyDefaultValue();
+				}
 			}
 			break;
 		}
diff --git a/Engine/Plugins/Experimental/DynamicMaterial/Source/DynamicMaterialEditor/Private/DetailsPanel/Widgets/SDMMaterialListExtensionWidget.cpp b/Engine/Plugins/Experimental/DynamicMaterial/Source/DynamicMaterialEditor/Private/DetailsPanel/Widgets/SDMMaterialListExtensionWidget.cpp
--- a/Engine/Plugins/Experimental/DynamicMaterial/Source/DynamicMaterialEditor/Private/DetailsPanel/Widgets/SDMMaterialListExtensionWidget.cpp
+++ b/Engine/Plugins/Experimental/DynamicMaterial/Source/DynamicMaterialEditor/Private/DetailsPanel/Widgets/SDMMaterialListExtensionWidget.cpp
@@ -2,15 +2,21 @@
 
 #include "DetailsPanel/Widgets/SDMMaterialListExtensionWidget.h"
 
+#include "AssetToolsModule.h"
 #include "Components/PrimitiveComponent.h"
 #include "DMObjectMaterialProperty.h"
 #include "DMWorldSubsystem.h"
 #include "DynamicMaterialEditorStyle.h"
 #include "Editor.h"
 #include "Engine/World.h"
+#include "Framework/Application/SlateApplication.h"
+#include "IAssetTools.h"
 #include "IDynamicMaterialEditorModule.h"
 #include "Material/DynamicMaterialInstance.h"
 #include "MaterialList.h"
+#include "Materials/Material.h"
+#include "Model/DynamicMaterialModelBase.h"
+#include "Modules/ModuleManager.h"
 #include "Widgets/Input/SButton.h"
 #include "Widgets/Layout/SBox.h"
 #include "Widgets/Text/STextBlock.h"
@@ -212,6 +218,26 @@ FReply SDMMaterialListExtensionWidget::OpenMaterialDesignerTab()
 		return FReply::Handled();
 	}
 
+	if (FSlateApplication::Get().GetModifierKeys().IsShiftDown())
+	{
+		UDynamicMaterialInstance* MaterialInstance = Cast<UDynamicMaterialInstance>(CurrentComponent->GetMaterial(MaterialItemView->GetMaterialListItem().SlotIndex));
+		if (!MaterialInstance)
+		{
+			return FReply::Handled();
+		}
+
+		UDynamicMaterialModelBase* MaterialModel = MaterialInstance->GetMaterialModelBase();
+		if (!MaterialModel)
+		{
+			return FReply::Handled();
+		}
+
+		IAssetTools& AssetTools = FAssetToolsModule::GetModule().Get();
+		AssetTools.OpenEditorForAssets({MaterialModel->GetGeneratedMaterial()});
+
+		return FReply::Handled();
+	}
+
 	IDynamicMaterialEditorModule::Get().OpenMaterialObjectProperty({CurrentComponent, MaterialItemView->GetMaterialListItem().SlotIndex},
 		CurrentComponent->GetWorld(), /* Invoke Tab */ true);
 
diff --git a/Engine/Plugins/Experimental/DynamicMaterial/Source/DynamicMaterialEditor/Private/UI/Utils/DMWidgetLibrary.cpp b/Engine/Plugins/Experimental/DynamicMaterial/Source/DynamicMaterialEditor/Private/UI/Utils/DMWidgetLibrary.cpp
--- a/Engine/Plugins/Experimental/DynamicMaterial/Source/DynamicMaterialEditor/Private/UI/Utils/DMWidgetLibrary.cpp
+++ b/Engine/Plugins/Experimental/DynamicMaterial/Source/DynamicMaterialEditor/Private/UI/Utils/DMWidgetLibrary.cpp
@@ -8,6 +8,7 @@
 #include "Model/DynamicMaterialModelBase.h"
 #include "Modules/ModuleManager.h"
 #include "PropertyEditorModule.h"
+#include "Utils/DMMaterialModelFunctionLibrary.h"
 #include "Utils/DMPrivate.h"
 
 IDMWidgetLibrary& IDMWidgetLibrary::Get()
@@ -209,7 +210,7 @@ FDMPropertyHandle FDMWidgetLibrary::CreatePropertyHandle(const FDMPropertyHandle
 		else if (InParams.Object->IsIn(InParams.PreviewMaterialModelBase))
 		{
 			const FString RelativePath = InParams.Object->GetPathName(InParams.PreviewMaterialModelBase);
-			OriginalObject = UE::DynamicMaterialEditor::Private::FindSubobject(InParams.OriginalMaterialModelBase, RelativePath);
+			OriginalObject = UDMMaterialModelFunctionLibrary::FindSubobject(InParams.OriginalMaterialModelBase, RelativePath);
 		}
 
 		if (OriginalObject)
diff --git a/Engine/Plugins/Experimental/DynamicMaterial/Source/DynamicMaterialEditor/Private/UI/Widgets/Editor/SDMMaterialComponentEditor.cpp b/Engine/Plugins/Experimental/DynamicMaterial/Source/DynamicMaterialEditor/Private/UI/Widgets/Editor/SDMMaterialComponentEditor.cpp
--- a/Engine/Plugins/Experimental/DynamicMaterial/Source/DynamicMaterialEditor/Private/UI/Widgets/Editor/SDMMaterialComponentEditor.cpp
+++ b/Engine/Plugins/Experimental/DynamicMaterial/Source/DynamicMaterialEditor/Private/UI/Widgets/Editor/SDMMaterialComponentEditor.cpp
@@ -194,19 +194,37 @@ TArray<FDMPropertyHandle> SDMMaterialComponentEditor::GetPropertyRows()
 
 	for (const FDMPropertyHandle& PropertyRow : PropertyRows)
 	{
+		TArray<UObject*> Outers;
+
 		if (PropertyRow.PreviewHandle.PropertyHandle.IsValid())
 		{
-			TArray<UObject*> Outers;
 			PropertyRow.PreviewHandle.PropertyHandle->GetOuterObjects(Outers);
+		}
+		else if (PropertyRow.PreviewHandle.PropertyRowGenerator.IsValid())
+		{
+			for (const TWeakObjectPtr<UObject>& WeakObject : PropertyRow.PreviewHandle.PropertyRowGenerator->GetSelectedObjects())
+			{
+				if (UObject* Object = WeakObject.Get())
+				{
+					Outers.Add(Object);
+				}
+			}
+		}
+
+		for (UObject* Outer : Outers)
+		{
+			UDMMaterialComponent* Component = Cast<UDMMaterialComponent>(Outer);
+
+			if (!Component)
+			{
+				Component = Outer->GetTypedOuter<UDMMaterialComponent>();
+			}
 
-			for (UObject* Outer : Outers)
+			if (Component)
 			{
-				if (UDMMaterialComponent* Component = Cast<UDMMaterialComponent>(Outer))
+				if (!Component->GetOnUpdate().IsBoundToObject(this))
 				{
-					if (!Component->GetOnUpdate().IsBoundToObject(this))
-					{
-						Component->GetOnUpdate().AddSP(this, &SDMMaterialComponentEditor::OnComponentUpdated);
-					}
+					Component->GetOnUpdate().AddSP(this, &SDMMaterialComponentEditor::OnComponentUpdated);
 				}
 			}
 		}
diff --git a/Engine/Plugins/Experimental/DynamicMaterial/Source/DynamicMaterialEditor/Private/UI/Widgets/SDMMaterialEditor.cpp b/Engine/Plugins/Experimental/DynamicMaterial/Source/DynamicMaterialEditor/Private/UI/Widgets/SDMMaterialEditor.cpp
--- a/Engine/Plugins/Experimental/DynamicMaterial/Source/DynamicMaterialEditor/Private/UI/Widgets/SDMMaterialEditor.cpp
+++ b/Engine/Plugins/Experimental/DynamicMaterial/Source/DynamicMaterialEditor/Private/UI/Widgets/SDMMaterialEditor.cpp
@@ -324,7 +324,7 @@ UDMMaterialComponent* SDMMaterialEditor::GetOriginalComponent(UDMMaterialCompone
 
 	const FString RelativePath = InPreviewComponent->GetPathName(PreviewMaterialModelBase);
 
-	return UE::DynamicMaterialEditor::Private::FindSubobject<UDMMaterialComponent>(OriginalMaterialModelBase, *RelativePath);
+	return UDMMaterialModelFunctionLibrary::FindSubobject<UDMMaterialComponent>(OriginalMaterialModelBase, *RelativePath);
 }
 
 EDMMaterialEditorMode SDMMaterialEditor::GetEditMode() const
diff --git a/Engine/Plugins/Experimental/DynamicMaterial/Source/DynamicMaterialEditor/Private/Utils/DMMaterialModelFunctionLibrary.cpp b/Engine/Plugins/Experimental/DynamicMaterial/Source/DynamicMaterialEditor/Private/Utils/DMMaterialModelFunctionLibrary.cpp
--- a/Engine/Plugins/Experimental/DynamicMaterial/Source/DynamicMaterialEditor/Private/Utils/DMMaterialModelFunctionLibrary.cpp
+++ b/Engine/Plugins/Experimental/DynamicMaterial/Source/DynamicMaterialEditor/Private/Utils/DMMaterialModelFunctionLibrary.cpp
@@ -593,4 +593,34 @@ void UDMMaterialModelFunctionLibrary::MirrorMaterialModel(UDynamicMaterialModelB
 	GEngine->NotifyToolsOfObjectReplacement(ReplacedObjects);
 }
 
+UObject* UDMMaterialModelFunctionLibrary::FindSubobject(UObject* InOuter, FStringView InPath)
+{
+	if (InPath.IsEmpty())
+	{
+		return InOuter;
+	}
+
+	int32 Index = INDEX_NONE;
+	InPath.FindChar(TEXT('.'), Index);
+
+	if (Index >= 0)
+	{
+		const FStringView FirstPath = InPath.Left(Index);
+		InPath.RightChopInline(Index + 1);
+
+		if (!FirstPath.IsEmpty())
+		{
+			const FString PathSection(FirstPath);
+
+			if (UObject* Subobject = StaticFindObjectFast(UObject::StaticClass(), InOuter, *PathSection))
+			{
+				return FindSubobject(Subobject, InPath);
+			}
+		}
+	}
+
+	const FString Path(InPath);
+	return StaticFindObjectFast(UObject::StaticClass(), InOuter, *Path);
+}
+
 #undef LOCTEXT_NAMESPACE
diff --git a/Engine/Plugins/Experimental/DynamicMaterial/Source/DynamicMaterialEditor/Private/Utils/DMPrivate.cpp b/Engine/Plugins/Experimental/DynamicMaterial/Source/DynamicMaterialEditor/Private/Utils/DMPrivate.cpp
--- a/Engine/Plugins/Experimental/DynamicMaterial/Source/DynamicMaterialEditor/Private/Utils/DMPrivate.cpp
+++ b/Engine/Plugins/Experimental/DynamicMaterial/Source/DynamicMaterialEditor/Private/Utils/DMPrivate.cpp
@@ -139,36 +139,6 @@ namespace UE::DynamicMaterialEditor::Private
 
 		return GetMaterialPropertyLongDisplayName(InMaterialProperty);
 	}
-
-	UObject* FindSubobject(UObject* InOuter, FStringView InPath)
-	{
-		if (InPath.IsEmpty())
-		{
-			return InOuter;
-		}
-
-		int32 Index = INDEX_NONE;
-		InPath.FindChar(TEXT('.'), Index);
-
-		if (Index >= 0)
-		{
-			const FStringView FirstPath = InPath.Left(Index);
-			InPath.RightChopInline(Index + 1);
-
-			if (!FirstPath.IsEmpty())
-			{
-				const FString PathSection(FirstPath);
-
-				if (UObject* Subobject = StaticFindObjectFast(UObject::StaticClass(), InOuter, *PathSection))
-				{
-					return FindSubobject(Subobject, InPath);
-				}
-			}
-		}
-
-		const FString Path(InPath);
-		return StaticFindObjectFast(UObject::StaticClass(), InOuter, *Path);
-	}
 }
 
 FDMMaterialLayerReference::FDMMaterialLayerReference()
diff --git a/Engine/Plugins/Experimental/DynamicMaterial/Source/DynamicMaterialEditor/Private/Utils/DMPrivate.h b/Engine/Plugins/Experimental/DynamicMaterial/Source/DynamicMaterialEditor/Private/Utils/DMPrivate.h
--- a/Engine/Plugins/Experimental/DynamicMaterial/Source/DynamicMaterialEditor/Private/Utils/DMPrivate.h
+++ b/Engine/Plugins/Experimental/DynamicMaterial/Source/DynamicMaterialEditor/Private/Utils/DMPrivate.h
@@ -39,14 +39,6 @@ namespace UE::DynamicMaterialEditor::Private
 	FText GetMaterialPropertyLongDisplayName(EDMMaterialPropertyType InMaterialProperty);
 
 	FText GetMaterialPropertyShortDisplayName(EDMMaterialPropertyType InMaterialProperty);
-
-	UObject* FindSubobject(UObject* InOuter, FStringView InPath);
-
-	template<typename InClassName>
-	InClassName* FindSubobject(UObject* InOuter, FStringView InPath)
-	{
-		return Cast<InClassName>(FindSubobject(InOuter, InPath));
-	}
 }
 
 struct FDMMaterialLayerReference
diff --git a/Engine/Plugins/Experimental/DynamicMaterial/Source/DynamicMaterialEditor/Public/Utils/DMMaterialModelFunctionLibrary.h b/Engine/Plugins/Experimental/DynamicMaterial/Source/DynamicMaterialEditor/Public/Utils/DMMaterialModelFunctionLibrary.h
--- a/Engine/Plugins/Experimental/DynamicMaterial/Source/DynamicMaterialEditor/Public/Utils/DMMaterialModelFunctionLibrary.h
+++ b/Engine/Plugins/Experimental/DynamicMaterial/Source/DynamicMaterialEditor/Public/Utils/DMMaterialModelFunctionLibrary.h
@@ -53,4 +53,12 @@ public:
 	 * objects, rather than recreating them. Additions and removals are also propagated.
 	 */
 	DYNAMICMATERIALEDITOR_API static void MirrorMaterialModel(UDynamicMaterialModelBase* InSource, UDynamicMaterialModelBase*& InTarget);
+
+	DYNAMICMATERIALEDITOR_API static UObject* FindSubobject(UObject* InOuter, FStringView InPath);
+
+	template<typename InClassName>
+	static InClassName* FindSubobject(UObject* InOuter, FStringView InPath)
+	{
+		return Cast<InClassName>(FindSubobject(InOuter, InPath));
+	}
 };
diff --git a/Engine/Plugins/Experimental/DynamicMaterialMediaStreamBridge/Source/DynamicMaterialMediaStreamBridge/Private/DMMaterialValueMediaStream.cpp b/Engine/Plugins/Experimental/DynamicMaterialMediaStreamBridge/Source/DynamicMaterialMediaStreamBridge/Private/DMMaterialValueMediaStream.cpp
--- a/Engine/Plugins/Experimental/DynamicMaterialMediaStreamBridge/Source/DynamicMaterialMediaStreamBridge/Private/DMMaterialValueMediaStream.cpp
+++ b/Engine/Plugins/Experimental/DynamicMaterialMediaStreamBridge/Source/DynamicMaterialMediaStreamBridge/Private/DMMaterialValueMediaStream.cpp
@@ -3,11 +3,24 @@
 #include "DMMaterialValueMediaStream.h"
 
 #include "DMMaterialValueMediaStreamDynamic.h"
+#include "Engine/Texture2D.h"
 #include "IMediaStreamPlayer.h"
 #include "MediaStream.h"
+#include "UObject/Package.h"
+#include "UObject/UObjectGlobals.h"
 
 #define LOCTEXT_NAMESPACE "DMMaterialValueMediaStream"
 
+namespace UE::DynamicMaterialMediaStreamBridge::Private
+{
+	constexpr const TCHAR* DefaultTexturePath = TEXT("/Script/Engine.Texture2D'/Engine/EditorResources/SceneManager.SceneManager'");
+
+	UTexture2D* GetDefaultTexture()
+	{
+		return LoadObject<UTexture2D>(GetTransientPackage(), DefaultTexturePath);
+	}
+}
+
 UDMMaterialValueMediaStream::UDMMaterialValueMediaStream()
 {
 	MediaStream = CreateDefaultSubobject<UMediaStream>("MediaStream");
@@ -114,6 +127,11 @@ bool UDMMaterialValueMediaStream::JsonDeserialize(const TSharedPtr<FJsonValue>&
 	return bSuccess;
 }
 
+void UDMMaterialValueMediaStream::ResetDefaultValue()
+{
+	DefaultValue = UE::DynamicMaterialMediaStreamBridge::Private::GetDefaultTexture();
+}
+
 UDMMaterialValueDynamic* UDMMaterialValueMediaStream::ToDynamic(UDynamicMaterialModelDynamic* InMaterialModelDynamic)
 {
 	UDMMaterialValueMediaStreamDynamic* ValueDynamic = UDMMaterialValueDynamic::CreateValueDynamic<UDMMaterialValueMediaStreamDynamic>(InMaterialModelDynamic, this);
@@ -145,6 +163,11 @@ void UDMMaterialValueMediaStream::PostEditorDuplicate(UDynamicMaterialModel* InM
 {
 	Super::PostEditorDuplicate(InMaterialModel, InParent);
 
+	if (IsValid(MediaStream))
+	{
+		MediaStream->EnsurePlayer();
+	}
+
 	UpdatePlayer();
 }
 
@@ -174,6 +197,12 @@ void UDMMaterialValueMediaStream::UpdateTextureFromMediaStream()
 		}
 	}
 
+	if ((!MediaTexture && IsDefaultValue()) || (MediaTexture == GetValue()))
+	{
+		return;
+	}
+
+	Update(this, EDMUpdateType::RefreshDetailView);
 	SetValue(MediaTexture);
 }
 
diff --git a/Engine/Plugins/Experimental/DynamicMaterialMediaStreamBridge/Source/DynamicMaterialMediaStreamBridge/Public/DMMaterialValueMediaStream.h b/Engine/Plugins/Experimental/DynamicMaterialMediaStreamBridge/Source/DynamicMaterialMediaStreamBridge/Public/DMMaterialValueMediaStream.h
--- a/Engine/Plugins/Experimental/DynamicMaterialMediaStreamBridge/Source/DynamicMaterialMediaStreamBridge/Public/DMMaterialValueMediaStream.h
+++ b/Engine/Plugins/Experimental/DynamicMaterialMediaStreamBridge/Source/DynamicMaterialMediaStreamBridge/Public/DMMaterialValueMediaStream.h
@@ -33,7 +33,8 @@ public:
 
 	//~ Begin UDMMaterialValue
 	virtual bool AllowEditValue() const { return false; }
-	DYNAMICMATERIALMEDIASTREAMBRIDGE_API  virtual UDMMaterialValueDynamic* ToDynamic(UDynamicMaterialModelDynamic* InMaterialModelDynamic) override;
+	DYNAMICMATERIALMEDIASTREAMBRIDGE_API virtual void ResetDefaultValue() override;
+	DYNAMICMATERIALMEDIASTREAMBRIDGE_API virtual UDMMaterialValueDynamic* ToDynamic(UDynamicMaterialModelDynamic* InMaterialModelDynamic) override;
 	//~ End UDMMaterialValue
 
 	//~ Begin UDMMaterialComponent
@@ -49,7 +50,7 @@ public:
 #endif
 
 protected:
-	UPROPERTY(VisibleAnywhere, BlueprintReadOnly, Category = "Material Designer|Media Stream",
+	UPROPERTY(VisibleAnywhere, BlueprintReadOnly, Category = "Material Designer|Media Stream", 
 		meta = (AllowPrivateAccess = "true", NotKeyframeable))
 	TObjectPtr<UMediaStream> MediaStream;
 
diff --git a/Engine/Plugins/Experimental/DynamicMaterialMediaStreamBridge/Source/DynamicMaterialMediaStreamBridgeEditor/Private/DMMaterialValueMediaStreamPropertyRowGenerator.cpp b/Engine/Plugins/Experimental/DynamicMaterialMediaStreamBridge/Source/DynamicMaterialMediaStreamBridgeEditor/Private/DMMaterialValueMediaStreamPropertyRowGenerator.cpp
--- a/Engine/Plugins/Experimental/DynamicMaterialMediaStreamBridge/Source/DynamicMaterialMediaStreamBridgeEditor/Private/DMMaterialValueMediaStreamPropertyRowGenerator.cpp
+++ b/Engine/Plugins/Experimental/DynamicMaterialMediaStreamBridge/Source/DynamicMaterialMediaStreamBridgeEditor/Private/DMMaterialValueMediaStreamPropertyRowGenerator.cpp
@@ -8,8 +8,10 @@
 #include "IDynamicMaterialEditorModule.h"
 #include "IMediaStreamPlayer.h"
 #include "MediaStream.h"
+#include "Model/DynamicMaterialModelBase.h"
 #include "UI/Utils/IDMWidgetLibrary.h"
 #include "UObject/Object.h"
+#include "Utils/DMMaterialModelFunctionLibrary.h"
 
 const TSharedRef<FDMMaterialValueMediaStreamPropertyRowGenerator>& FDMMaterialValueMediaStreamPropertyRowGenerator::Get()
 {
@@ -57,6 +59,7 @@ void FDMMaterialValueMediaStreamPropertyRowGenerator::AddSourceProperty(FDMCompo
 	FDMPropertyHandle& NewHandle = InParams.PropertyRows->Add_GetRef(
 		IDMWidgetLibrary::Get().GetPropertyHandle(InParams.CreatePropertyHandleParams(TEXT("Source")))
 	);
+	NewHandle.CategoryOverrideName = TEXT("Media Stream");
 }
 
 void FDMMaterialValueMediaStreamPropertyRowGenerator::AddPlayerProperty(FDMComponentPropertyRowGeneratorParams& InParams)
@@ -70,12 +73,48 @@ void FDMMaterialValueMediaStreamPropertyRowGenerator::AddPlayerProperty(FDMCompo
 		return;
 	}
 
-	const FName VideoPlayerCategory = TEXT("Video Player");
+	const FName PreviewCategory = TEXT("Preview Controls");
+	const FName SourceCategory = TEXT("Source Controls");
+	const FName VideoPlayerCategory = TEXT("Media Stream");
+
+	const FString RelativePath = InParams.Object->GetPathName(InParams.PreviewMaterialModelBase);
+	UMediaStream* OriginalMediaStream = UDMMaterialModelFunctionLibrary::FindSubobject<UMediaStream>(InParams.OriginalMaterialModelBase, RelativePath);
 
 	FDMPropertyHandle& ControlHandle = InParams.PropertyRows->Add_GetRef(
 		IDMWidgetLibrary::Get().GetPropertyHandle(InParams.CreatePropertyHandleParams(TEXT("Control")))
 	);
-	ControlHandle.CategoryOverrideName = VideoPlayerCategory;
+	ControlHandle.CategoryOverrideName = PreviewCategory;
+
+	if (OriginalMediaStream)
+	{
+		FDMComponentPropertyRowGeneratorParams SourceParams = InParams;
+		SourceParams.Object = OriginalMediaStream;
+
+		FDMPropertyHandle& SourceControlHandle = InParams.PropertyRows->Add_GetRef(
+				IDMWidgetLibrary::Get().GetPropertyHandle(SourceParams.CreatePropertyHandleParams(TEXT("Control")))
+			);
+		SourceControlHandle.CategoryOverrideName = SourceCategory;
+	}
+
+	{
+		FDMComponentPropertyRowGeneratorParams PlayerParams = InParams;
+		PlayerParams.Object = Player->_getUObject();
+
+		FDMPropertyHandle& PlaylistIndexHandle = InParams.PropertyRows->Add_GetRef(
+			IDMWidgetLibrary::Get().GetPropertyHandle(PlayerParams.CreatePropertyHandleParams(TEXT("PlaylistIndex")))
+		);
+		PlaylistIndexHandle.CategoryOverrideName = VideoPlayerCategory;
+
+		FDMPropertyHandle& PlaybackStateHandle = InParams.PropertyRows->Add_GetRef(
+			IDMWidgetLibrary::Get().GetPropertyHandle(PlayerParams.CreatePropertyHandleParams(TEXT("PlaybackState")))
+		);
+		PlaybackStateHandle.CategoryOverrideName = VideoPlayerCategory;
+
+		FDMPropertyHandle& RequestedSeekFrameHandle = InParams.PropertyRows->Add_GetRef(
+			IDMWidgetLibrary::Get().GetPropertyHandle(PlayerParams.CreatePropertyHandleParams(TEXT("RequestedSeekFrame")))
+		);
+		RequestedSeekFrameHandle.CategoryOverrideName = VideoPlayerCategory;
+	}
 
 	FDMPropertyHandle& TextureHandle = InParams.PropertyRows->Add_GetRef(
 		IDMWidgetLibrary::Get().GetPropertyHandle(InParams.CreatePropertyHandleParams(TEXT("MediaTexture")))
diff --git a/Engine/Plugins/Experimental/MediaStream/Source/MediaStream/Private/MediaStream.cpp b/Engine/Plugins/Experimental/MediaStream/Source/MediaStream/Private/MediaStream.cpp
--- a/Engine/Plugins/Experimental/MediaStream/Source/MediaStream/Private/MediaStream.cpp
+++ b/Engine/Plugins/Experimental/MediaStream/Source/MediaStream/Private/MediaStream.cpp
@@ -50,6 +50,33 @@ TScriptInterface<IMediaStreamPlayer> UMediaStream::GetPlayer() const
 	return PlayerObject;
 }
 
+bool UMediaStream::EnsurePlayer(bool bInForceRecreatePlayer)
+{
+	if (bInForceRecreatePlayer)
+	{
+		FMediaStreamSource SourceTemp = Source;
+		Source = FMediaStreamSource();
+		
+		if (IsValid(PlayerObject))
+		{
+			if (IMediaStreamPlayer* Player = Cast<IMediaStreamPlayer>(PlayerObject))
+			{
+				Player->Deinitialize();
+				Player->OnSourceChanged(Source);
+				PlayerObject = nullptr;
+			}
+		}
+
+		SetSource(SourceTemp);
+	}
+	else if (!PlayerObject && UMediaStreamSourceBlueprintLibrary::IsValidMediaSource(Source))
+	{
+		ApplySource();
+	}
+
+	return !!PlayerObject;
+}
+
 UMediaStream::FOnSourceChanged& UMediaStream::GetOnSourceChanged()
 {
 	return OnSourceChanged;
@@ -219,7 +246,6 @@ bool UMediaStream::ApplySource()
 			{
 				Player->Deinitialize();
 				Player->OnSourceChanged(Source);
-
 			}
 		}
 
diff --git a/Engine/Plugins/Experimental/MediaStream/Source/MediaStream/Private/Players/MediaStreamLocalPlayer.cpp b/Engine/Plugins/Experimental/MediaStream/Source/MediaStream/Private/Players/MediaStreamLocalPlayer.cpp
--- a/Engine/Plugins/Experimental/MediaStream/Source/MediaStream/Private/Players/MediaStreamLocalPlayer.cpp
+++ b/Engine/Plugins/Experimental/MediaStream/Source/MediaStream/Private/Players/MediaStreamLocalPlayer.cpp
@@ -82,7 +82,7 @@ void UMediaStreamLocalPlayer::PostEditChangeProperty(FPropertyChangedEvent& InPr
 
 	static const FName PlayerConfigName = GET_MEMBER_NAME_CHECKED(UMediaStreamLocalPlayer, PlayerConfig);
 	static const FName TextureConfigName = GET_MEMBER_NAME_CHECKED(UMediaStreamLocalPlayer, TextureConfig);
-	static const FName RequestedSeekTimeName = GET_MEMBER_NAME_CHECKED(UMediaStreamLocalPlayer, RequestedSeekTime);
+	static const FName RequestedSeekTimeName = GET_MEMBER_NAME_CHECKED(UMediaStreamLocalPlayer, RequestedSeekFrame);
 	static const FName PlaybackStateName = GET_MEMBER_NAME_CHECKED(UMediaStreamLocalPlayer, PlaybackState);
 	static const FName PlaylistIndexName = GET_MEMBER_NAME_CHECKED(UMediaStreamLocalPlayer, PlaylistIndex);
 
@@ -98,7 +98,7 @@ void UMediaStreamLocalPlayer::PostEditChangeProperty(FPropertyChangedEvent& InPr
 	}
 	else if (PropertyName == RequestedSeekTimeName)
 	{
-		ApplyRequestedSeekTime();
+		ApplyRequestedSeekFrame();
 	}
 	else if (PropertyName == PlaybackStateName)
 	{
@@ -255,26 +255,67 @@ void UMediaStreamLocalPlayer::ApplyPlayerConfig()
 
 float UMediaStreamLocalPlayer::GetRequestedSeekTime() const
 {
-	return RequestedSeekTime;
+	if (!IsValid(Player))
+	{
+		return 0.f;
+	}
+
+	float FrameRate = Player->GetVideoTrackFrameRate(INDEX_NONE, INDEX_NONE);
+
+	if (FMath::IsNearlyZero(FrameRate))
+	{
+		return 0.f;
+	}
+
+	return static_cast<float>(RequestedSeekFrame) / FrameRate;
 }
 
 bool UMediaStreamLocalPlayer::SetRequestedSeekTime(float InTime)
 {
-	RequestedSeekTime = InTime;
+	if (!IsValid(Player))
+	{
+		return false;
+	}
+
+	float FrameRate = Player->GetVideoTrackFrameRate(INDEX_NONE, INDEX_NONE);
+
+	if (FMath::IsNearlyZero(FrameRate))
+	{
+		return false;
+	}
+
+	RequestedSeekFrame = FMath::FloorToInt(InTime * FrameRate);
+
+	return ApplyRequestedSeekFrame();
+}
+
+int32 UMediaStreamLocalPlayer::GetRequestedSeekFrame() const
+{
+	return RequestedSeekFrame;;
+}
+
+bool UMediaStreamLocalPlayer::SetRequestedSeekFrame(int32 InFrame)
+{
+	RequestedSeekFrame = InFrame;
 
-	return ApplyRequestedSeekTime();
+	return ApplyRequestedSeekFrame();
 }
 
-bool UMediaStreamLocalPlayer::ApplyRequestedSeekTime()
+bool UMediaStreamLocalPlayer::ApplyRequestedSeekFrame()
 {
 	bool bSuccess = false;
 
-	if (RequestedSeekTime >= 0 && IsValid(Player) && Player->SupportsSeeking())
+	if (RequestedSeekFrame >= 0 && IsValid(Player) && Player->SupportsSeeking())
 	{
-		bSuccess = Player->Seek(RequestedSeekTime);
+		const float FrameRate = Player->GetVideoTrackFrameRate(INDEX_NONE, INDEX_NONE);
+
+		if (FrameRate > 0.f)
+		{
+			bSuccess = Player->Seek(static_cast<float>(RequestedSeekFrame) / FrameRate);
+		}
 	}
 
-	RequestedSeekTime = -1.f;
+	RequestedSeekFrame = INDEX_NONE;
 	return bSuccess;
 }
 
diff --git a/Engine/Plugins/Experimental/MediaStream/Source/MediaStream/Private/Players/MediaStreamLocalPlayer.h b/Engine/Plugins/Experimental/MediaStream/Source/MediaStream/Private/Players/MediaStreamLocalPlayer.h
--- a/Engine/Plugins/Experimental/MediaStream/Source/MediaStream/Private/Players/MediaStreamLocalPlayer.h
+++ b/Engine/Plugins/Experimental/MediaStream/Source/MediaStream/Private/Players/MediaStreamLocalPlayer.h
@@ -20,7 +20,7 @@ class UMediaStreamLocalPlayer : public UObject, public IMediaStreamPlayer
 
 public:
 	UMediaStreamLocalPlayer() = default;
-	virtual ~UMediaStreamLocalPlayer() override;	
+	virtual ~UMediaStreamLocalPlayer() override;
 
 	UFUNCTION(BlueprintCallable, DisplayName = "Set Requested Seek Time", Category = "Media Stream|Player")
 	void BP_SetRequestedSeekTime(float InSeekTime)
@@ -28,6 +28,12 @@ public:
 		SetRequestedSeekTime(InSeekTime);
 	}
 
+	UFUNCTION(BlueprintCallable, DisplayName = "Set Requested Seek Frame", Category = "Media Stream|Player")
+	void BP_SetRequestedSeekFrame(int32 InSeekFrame)
+	{
+		SetRequestedSeekFrame(InSeekFrame);
+	}
+
 	UFUNCTION(BlueprintCallable, DisplayName = "Set Playback State", Category = "Media Stream|Player")
 	void BP_SetPlaybackState(EMediaStreamPlaybackState InState)
 	{
@@ -59,6 +65,8 @@ public:
 	virtual bool SetPlaylistIndex(int32 InIndex) override;
 	virtual float GetRequestedSeekTime() const override;
 	virtual bool SetRequestedSeekTime(float InTime) override;
+	virtual int32 GetRequestedSeekFrame() const override;
+	virtual bool SetRequestedSeekFrame(int32 InFrame) override;
 	virtual EMediaStreamPlaybackState GetPlaybackState() const override;
 	virtual bool SetPlaybackState(EMediaStreamPlaybackState InState) override;
 	virtual int32 GetPlaylistIndex() const override;
@@ -103,10 +111,10 @@ protected:
 	UPROPERTY(EditAnywhere, BlueprintReadWrite, Category = "Media Stream|Player", meta = (AllowPrivateAccess = true))
 	FMediaStreamPlayerConfig PlayerConfig;
 
-	/** Attempts to seek to this time. Any value below 0 is ignored. */
-	UPROPERTY(EditAnywhere, BlueprintReadWrite, Setter = BP_SetRequestedSeekTime, Getter, BlueprintSetter = BP_SetRequestedSeekTime,
+	/** Attempts to seek to this frame. Any value below 0 is ignored. */
+	UPROPERTY(EditAnywhere, BlueprintReadWrite, Setter = BP_SetRequestedSeekFrame, Getter, BlueprintSetter = BP_SetRequestedSeekFrame,
 		Category = "Media Stream|Player", Interp, meta = (AllowPrivateAccess = true))
-	float RequestedSeekTime = -1.f;
+	int32 RequestedSeekFrame = INDEX_NONE;
 
 	UPROPERTY(EditAnywhere, BlueprintReadWrite, Setter = BP_SetPlaybackState, Getter, BlueprintSetter = BP_SetPlaybackState,
 		Category = "Media Stream|Player", Interp, meta = (AllowPrivateAccess = true))
@@ -115,7 +123,7 @@ protected:
 	/** The index playing in the playlist. -1 to ignore. */
 	UPROPERTY(EditAnywhere, BlueprintReadWrite, Setter = BP_SetPlaylistIndex, Getter, BlueprintSetter = BP_SetPlaylistIndex,
 		Category = "Media Stream|Player", Interp, meta = (AllowPrivateAccess = true))
-	int32 PlaylistIndex = -1;
+	int32 PlaylistIndex = INDEX_NONE;
 
 	FMediaStreamTextureConfig TextureConfig_PreUndo;
 	FMediaStreamPlayerConfig PlayerConfig_PreUndo;
@@ -163,7 +171,7 @@ protected:
 	 * Attempts to set the seek time on the currently active Media Player.
 	 * @return True if successful.
 	 */
-	bool ApplyRequestedSeekTime();
+	bool ApplyRequestedSeekFrame();
 
 	/**
 	 * Will either set the current Rate from the player config or 0 for paused.
diff --git a/Engine/Plugins/Experimental/MediaStream/Source/MediaStream/Private/Players/MediaStreamProxyPlayer.cpp b/Engine/Plugins/Experimental/MediaStream/Source/MediaStream/Private/Players/MediaStreamProxyPlayer.cpp
--- a/Engine/Plugins/Experimental/MediaStream/Source/MediaStream/Private/Players/MediaStreamProxyPlayer.cpp
+++ b/Engine/Plugins/Experimental/MediaStream/Source/MediaStream/Private/Players/MediaStreamProxyPlayer.cpp
@@ -249,6 +249,35 @@ bool UMediaStreamProxyPlayer::SetRequestedSeekTime(float InTime)
 	return false;
 }
 
+int32 UMediaStreamProxyPlayer::GetRequestedSeekFrame() const
+{
+	if (UMediaStream* MediaStream = GetSourceStream())
+	{
+		if (TScriptInterface<IMediaStreamPlayer> MediaStreamPlayer = MediaStream->GetPlayer())
+		{
+			return MediaStreamPlayer->GetRequestedSeekFrame();
+		}
+	}
+
+	return 0;
+}
+
+bool UMediaStreamProxyPlayer::SetRequestedSeekFrame(int32 InFrame)
+{
+	if (!bReadOnly)
+	{
+		if (UMediaStream* MediaStream = GetSourceStream())
+		{
+			if (TScriptInterface<IMediaStreamPlayer> MediaStreamPlayer = MediaStream->GetPlayer())
+			{
+				return MediaStreamPlayer->SetRequestedSeekFrame(InFrame);
+			}
+		}
+	}
+
+	return false;
+}
+
 EMediaStreamPlaybackState UMediaStreamProxyPlayer::GetPlaybackState() const
 {
 	if (UMediaStream* MediaStream = GetSourceStream())
diff --git a/Engine/Plugins/Experimental/MediaStream/Source/MediaStream/Private/Players/MediaStreamProxyPlayer.h b/Engine/Plugins/Experimental/MediaStream/Source/MediaStream/Private/Players/MediaStreamProxyPlayer.h
--- a/Engine/Plugins/Experimental/MediaStream/Source/MediaStream/Private/Players/MediaStreamProxyPlayer.h
+++ b/Engine/Plugins/Experimental/MediaStream/Source/MediaStream/Private/Players/MediaStreamProxyPlayer.h
@@ -41,6 +41,8 @@ public:
 	virtual bool SetPlaylistIndex(int32 InIndex) override;
 	virtual float GetRequestedSeekTime() const override;
 	virtual bool SetRequestedSeekTime(float InTime) override;
+	virtual int32 GetRequestedSeekFrame() const override;
+	virtual bool SetRequestedSeekFrame(int32 InFrame) override;
 	virtual EMediaStreamPlaybackState GetPlaybackState() const override;
 	virtual bool SetPlaybackState(EMediaStreamPlaybackState InState) override;
 	virtual int32 GetPlaylistIndex() const override;
diff --git a/Engine/Plugins/Experimental/MediaStream/Source/MediaStream/Public/IMediaStreamPlayer.h b/Engine/Plugins/Experimental/MediaStream/Source/MediaStream/Public/IMediaStreamPlayer.h
--- a/Engine/Plugins/Experimental/MediaStream/Source/MediaStream/Public/IMediaStreamPlayer.h
+++ b/Engine/Plugins/Experimental/MediaStream/Source/MediaStream/Public/IMediaStreamPlayer.h
@@ -129,6 +129,20 @@ public:
 	UFUNCTION(BlueprintCallable, Category = "Media Stream|Player")
 	virtual bool SetRequestedSeekTime(float InTime) = 0;
 
+	/**
+	 * @return The frame offset in the media that has been requested.
+	 */
+	UFUNCTION(BlueprintCallable, Category = "Media Stream|Player")
+	virtual int32 GetRequestedSeekFrame() const = 0;
+
+	/**
+	 * Sets the frame in the currently playing player, if possible.
+	 * @param InFrame The frame offset in the media to request.
+	 * @return True on success. False on failure or invalid player.
+	 */
+	UFUNCTION(BlueprintCallable, Category = "Media Stream|Player")
+	virtual bool SetRequestedSeekFrame(int32 InFrame) = 0;
+
 	/**
 	 * @return Gets the media player playback state.
 	 */
diff --git a/Engine/Plugins/Experimental/MediaStream/Source/MediaStream/Public/MediaStream.h b/Engine/Plugins/Experimental/MediaStream/Source/MediaStream/Public/MediaStream.h
--- a/Engine/Plugins/Experimental/MediaStream/Source/MediaStream/Public/MediaStream.h
+++ b/Engine/Plugins/Experimental/MediaStream/Source/MediaStream/Public/MediaStream.h
@@ -56,6 +56,14 @@ public:
 	UFUNCTION(BlueprintPure, Category = "Media Stream")
 	MEDIASTREAM_API TScriptInterface<IMediaStreamPlayer> GetPlayer() const;
 
+	/**
+	 * Creates a player if it doesn't exist.
+	 * @param  bInForceRecreatePlayer Will remove and recreate the player.
+	 * @return True if a valid player exists.
+	 */
+	UFUNCTION(BlueprintCallable, Category = "Media Stream")
+	MEDIASTREAM_API bool EnsurePlayer(bool bInForceRecreatePlayer = false);
+
 	/**
 	 * Called when the Media Stream Player changes.
 	 */
diff --git a/Engine/Plugins/Experimental/MediaStream/Source/MediaStreamEditor/Private/MediaStreamEditorSequencerLibrary.cpp b/Engine/Plugins/Experimental/MediaStream/Source/MediaStreamEditor/Private/MediaStreamEditorSequencerLibrary.cpp
--- a/Engine/Plugins/Experimental/MediaStream/Source/MediaStreamEditor/Private/MediaStreamEditorSequencerLibrary.cpp
+++ b/Engine/Plugins/Experimental/MediaStream/Source/MediaStreamEditor/Private/MediaStreamEditorSequencerLibrary.cpp
@@ -72,6 +72,11 @@ bool FMediaStreamEditorSequencerLibrary::CanAddTrack(UMediaStream* InMediaStream
 		return false;
 	}
 
+	if (InMediaStream->IsIn(GetTransientPackage()))
+	{
+		return false;
+	}
+
 	UObject* MediaObject = InMediaStream->GetSource().Object;
 
 	// Can't add a proxy stream to a track.
