Material Designer: Material Stream integration fixes
- Corrected the icon when a video is the current source.
- Made sure that changes to the video propagate to the source material.
- Made sure that the source material loads its texture.

diff --git a/Engine/Plugins/Experimental/DynamicMaterial/Source/DynamicMaterialEditor/Private/Components/MaterialStageExpressions/DMMSETextureSampleBase.cpp b/Engine/Plugins/Experimental/DynamicMaterial/Source/DynamicMaterialEditor/Private/Components/MaterialStageExpressions/DMMSETextureSampleBase.cpp
--- a/Engine/Plugins/Experimental/DynamicMaterial/Source/DynamicMaterialEditor/Private/Components/MaterialStageExpressions/DMMSETextureSampleBase.cpp
+++ b/Engine/Plugins/Experimental/DynamicMaterial/Source/DynamicMaterialEditor/Private/Components/MaterialStageExpressions/DMMSETextureSampleBase.cpp
@@ -259,6 +259,27 @@ FText UDMMaterialStageExpressionTextureSampleBase::GetComponentDescription() con
 	return Super::GetComponentDescription();
 }
 
+FSlateIcon UDMMaterialStageExpressionTextureSampleBase::GetComponentIcon() const
+{
+	if (UDMMaterialStage* Stage = GetStage())
+	{
+		const TArray<UDMMaterialStageInput*>& Inputs = Stage->GetInputs();
+
+		for (UDMMaterialStageInput* StageInput : Inputs)
+		{
+			if (UDMMaterialStageInputValue* InputValue = Cast<UDMMaterialStageInputValue>(StageInput))
+			{
+				if (UDMMaterialValue* Value = Cast<UDMMaterialValueTexture>(InputValue->GetValue()))
+				{
+					return Value->GetComponentIcon();
+				}
+			}
+		}
+	}
+
+	return Super::GetComponentIcon();
+}
+
 void UDMMaterialStageExpressionTextureSampleBase::SetClampTextureEnabled(bool bInValue)
 {
 	if (bClampTexture == bInValue)
diff --git a/Engine/Plugins/Experimental/DynamicMaterial/Source/DynamicMaterialEditor/Private/DynamicMaterialEditorSettings.h b/Engine/Plugins/Experimental/DynamicMaterial/Source/DynamicMaterialEditor/Private/DynamicMaterialEditorSettings.h
--- a/Engine/Plugins/Experimental/DynamicMaterial/Source/DynamicMaterialEditor/Private/DynamicMaterialEditorSettings.h
+++ b/Engine/Plugins/Experimental/DynamicMaterial/Source/DynamicMaterialEditor/Private/DynamicMaterialEditorSettings.h
@@ -202,14 +202,15 @@ public:
 	UPROPERTY(Config, EditAnywhere, BlueprintReadOnly, Category = "Editor")
 	bool bAddDetailsPanelButton;
 
+protected:
 	/**
 	 * When properties on the preview material are changed, if they are parameter-based then automatically
-	 * also set those values on the source material (in the world or asset).
+	 * also set those values on the source material (in the world or asset). This value is only used when Live Edit
+	 * mode is disabled.
 	 */
 	UPROPERTY(Config, EditAnywhere, BlueprintReadOnly, Category = "Editor")
 	bool bAutomaticallyCopyParametersToSourceMaterial;
 
-protected:
 	/**
 	 * When a structural material change is made, whether to automatically recompile the preview material to show this change.
 	 */
@@ -229,6 +230,17 @@ protected:
 	EDMLiveEditMode LiveEditMode;
 
 public:
+	UFUNCTION(BlueprintPure, Category = "Material Designer")
+	bool ShouldAutomaticallyCopyParametersToSourceMaterial() const
+	{
+		if (ShouldAutomaticallyApplyToSourceOnPreviewCompile())
+		{
+			return true;
+		}
+
+		return bAutomaticallyCopyParametersToSourceMaterial;
+	}
+
 	UFUNCTION(BlueprintPure, Category = "Material Designer")
 	bool ShouldAutomaticallyCompilePreviewMaterial() const
 	{
diff --git a/Engine/Plugins/Experimental/DynamicMaterial/Source/DynamicMaterialEditor/Private/UI/Widgets/Editor/SDMMaterialComponentEditor.cpp b/Engine/Plugins/Experimental/DynamicMaterial/Source/DynamicMaterialEditor/Private/UI/Widgets/Editor/SDMMaterialComponentEditor.cpp
--- a/Engine/Plugins/Experimental/DynamicMaterial/Source/DynamicMaterialEditor/Private/UI/Widgets/Editor/SDMMaterialComponentEditor.cpp
+++ b/Engine/Plugins/Experimental/DynamicMaterial/Source/DynamicMaterialEditor/Private/UI/Widgets/Editor/SDMMaterialComponentEditor.cpp
@@ -117,7 +117,7 @@ void SDMMaterialComponentEditor::OnComponentUpdated(UDMMaterialComponent* InComp
 
 	if (UDynamicMaterialEditorSettings* Settings = UDynamicMaterialEditorSettings::Get())
 	{
-		if (Settings->bAutomaticallyCopyParametersToSourceMaterial)
+		if (Settings->ShouldAutomaticallyCopyParametersToSourceMaterial())
 		{
 			if (InComponent)
 			{
diff --git a/Engine/Plugins/Experimental/DynamicMaterial/Source/DynamicMaterialEditor/Private/UI/Widgets/Editor/SDMMaterialGlobalSettingsEditor.cpp b/Engine/Plugins/Experimental/DynamicMaterial/Source/DynamicMaterialEditor/Private/UI/Widgets/Editor/SDMMaterialGlobalSettingsEditor.cpp
--- a/Engine/Plugins/Experimental/DynamicMaterial/Source/DynamicMaterialEditor/Private/UI/Widgets/Editor/SDMMaterialGlobalSettingsEditor.cpp
+++ b/Engine/Plugins/Experimental/DynamicMaterial/Source/DynamicMaterialEditor/Private/UI/Widgets/Editor/SDMMaterialGlobalSettingsEditor.cpp
@@ -128,7 +128,7 @@ void SDMMaterialGlobalSettingsEditor::OnComponentUpdated(UDMMaterialComponent* I
 {
 	if (UDynamicMaterialEditorSettings* Settings = UDynamicMaterialEditorSettings::Get())
 	{
-		if (Settings->bAutomaticallyCopyParametersToSourceMaterial)
+		if (Settings->ShouldAutomaticallyCopyParametersToSourceMaterial())
 		{
 			if (InComponent)
 			{
diff --git a/Engine/Plugins/Experimental/DynamicMaterial/Source/DynamicMaterialEditor/Private/UI/Widgets/Editor/SDMMaterialProperties.cpp b/Engine/Plugins/Experimental/DynamicMaterial/Source/DynamicMaterialEditor/Private/UI/Widgets/Editor/SDMMaterialProperties.cpp
--- a/Engine/Plugins/Experimental/DynamicMaterial/Source/DynamicMaterialEditor/Private/UI/Widgets/Editor/SDMMaterialProperties.cpp
+++ b/Engine/Plugins/Experimental/DynamicMaterial/Source/DynamicMaterialEditor/Private/UI/Widgets/Editor/SDMMaterialProperties.cpp
@@ -201,7 +201,7 @@ void SDMMaterialProperties::OnComponentUpdated(UDMMaterialComponent* InComponent
 {
 	if (UDynamicMaterialEditorSettings* Settings = UDynamicMaterialEditorSettings::Get())
 	{
-		if (Settings->bAutomaticallyCopyParametersToSourceMaterial)
+		if (Settings->ShouldAutomaticallyCopyParametersToSourceMaterial())
 		{
 			if (InComponent)
 			{
diff --git a/Engine/Plugins/Experimental/DynamicMaterial/Source/DynamicMaterialEditor/Private/UI/Widgets/Editor/SDMMaterialSlotEditor.cpp b/Engine/Plugins/Experimental/DynamicMaterial/Source/DynamicMaterialEditor/Private/UI/Widgets/Editor/SDMMaterialSlotEditor.cpp
--- a/Engine/Plugins/Experimental/DynamicMaterial/Source/DynamicMaterialEditor/Private/UI/Widgets/Editor/SDMMaterialSlotEditor.cpp
+++ b/Engine/Plugins/Experimental/DynamicMaterial/Source/DynamicMaterialEditor/Private/UI/Widgets/Editor/SDMMaterialSlotEditor.cpp
@@ -990,7 +990,7 @@ void SDMMaterialSlotEditor::OnOpacityUpdated(UDMMaterialComponent* InComponent,
 {
 	if (UDynamicMaterialEditorSettings* Settings = UDynamicMaterialEditorSettings::Get())
 	{
-		if (Settings->bAutomaticallyCopyParametersToSourceMaterial)
+		if (Settings->ShouldAutomaticallyCopyParametersToSourceMaterial())
 		{
 			if (InComponent)
 			{
diff --git a/Engine/Plugins/Experimental/DynamicMaterial/Source/DynamicMaterialEditor/Public/Components/MaterialStageExpressions/DMMSETextureSampleBase.h b/Engine/Plugins/Experimental/DynamicMaterial/Source/DynamicMaterialEditor/Public/Components/MaterialStageExpressions/DMMSETextureSampleBase.h
--- a/Engine/Plugins/Experimental/DynamicMaterial/Source/DynamicMaterialEditor/Public/Components/MaterialStageExpressions/DMMSETextureSampleBase.h
+++ b/Engine/Plugins/Experimental/DynamicMaterial/Source/DynamicMaterialEditor/Public/Components/MaterialStageExpressions/DMMSETextureSampleBase.h
@@ -40,6 +40,7 @@ public:
 
 	//~ Begin UDMMaterialComponent
 	virtual FText GetComponentDescription() const override;
+	virtual FSlateIcon GetComponentIcon() const override;
 	//~ End UDMMaterialComponent;
 
 	//~ Begin UObject
diff --git a/Engine/Plugins/Experimental/DynamicMaterialMediaStreamBridge/Source/DynamicMaterialMediaStreamBridge/Private/DMMaterialValueMediaStream.cpp b/Engine/Plugins/Experimental/DynamicMaterialMediaStreamBridge/Source/DynamicMaterialMediaStreamBridge/Private/DMMaterialValueMediaStream.cpp
--- a/Engine/Plugins/Experimental/DynamicMaterialMediaStreamBridge/Source/DynamicMaterialMediaStreamBridge/Private/DMMaterialValueMediaStream.cpp
+++ b/Engine/Plugins/Experimental/DynamicMaterialMediaStreamBridge/Source/DynamicMaterialMediaStreamBridge/Private/DMMaterialValueMediaStream.cpp
@@ -173,6 +173,15 @@ void UDMMaterialValueMediaStream::PostEditorDuplicate(UDynamicMaterialModel* InM
 	}
 
 	UpdatePlayer();
+
+	if (MediaStream)
+	{
+		if (IMediaStreamPlayer* MediaStreamPlayer = MediaStream->GetPlayer().GetInterface())
+		{
+			MediaStreamPlayer->ApplyTextureConfig();
+			MediaStreamPlayer->ApplyPlayerConfig();
+		}
+	}
 }
 
 FSlateIcon UDMMaterialValueMediaStream::GetComponentIcon() const
@@ -186,7 +195,7 @@ void UDMMaterialValueMediaStream::PostEditChangeProperty(FPropertyChangedEvent&
 
 	if (InPropertyChangedEvent.GetMemberPropertyName() == UMediaStream::GetSourcePropertyName())
 	{
-		Update(this, EDMUpdateType::RefreshDetailView);
+		Update(this, EDMUpdateType::Structure | EDMUpdateType::RefreshDetailView);
 	}
 
 	UpdatePlayer();
@@ -260,7 +269,39 @@ void UDMMaterialValueMediaStream::UpdatePlayer()
 
 void UDMMaterialValueMediaStream::CopyParametersFrom_Implementation(UObject* InOther)
 {
-	// Overridden to do nothing
+	UDMMaterialValueMediaStream* Other = Cast<UDMMaterialValueMediaStream>(InOther);
+
+	UMediaStream* OtherMediaStream = Other->GetMediaStream();
+
+	if (!MediaStream || !OtherMediaStream)
+	{
+		return;
+	}
+
+	IMediaStreamPlayer* OtherPlayer = OtherMediaStream->GetPlayer().GetInterface();
+
+	if (OtherPlayer && !OtherPlayer->IsReadOnly())
+	{
+		if (IMediaStreamPlayer* Player = MediaStream->GetPlayer().GetInterface())
+		{
+			if (OtherPlayer->GetTextureConfig() != Player->GetTextureConfig())
+			{
+				OtherPlayer->SetTextureConfig(Player->GetTextureConfig());
+			}
+
+			if (OtherPlayer->GetPlayerConfig() != Player->GetPlayerConfig())
+			{
+				OtherPlayer->SetPlayerConfig(Player->GetPlayerConfig());
+			}
+		}
+	}
+
+	if (OtherMediaStream->GetSource() != MediaStream->GetSource())
+	{
+		OtherMediaStream->SetSource(MediaStream->GetSource());
+	}
+
+	UpdatePlayer();
 }
 
 void UDMMaterialValueMediaStream::OnComponentAdded()
diff --git a/Engine/Plugins/Experimental/DynamicMaterialMediaStreamBridge/Source/DynamicMaterialMediaStreamBridge/Private/DMMaterialValueMediaStreamDynamic.cpp b/Engine/Plugins/Experimental/DynamicMaterialMediaStreamBridge/Source/DynamicMaterialMediaStreamBridge/Private/DMMaterialValueMediaStreamDynamic.cpp
--- a/Engine/Plugins/Experimental/DynamicMaterialMediaStreamBridge/Source/DynamicMaterialMediaStreamBridge/Private/DMMaterialValueMediaStreamDynamic.cpp
+++ b/Engine/Plugins/Experimental/DynamicMaterialMediaStreamBridge/Source/DynamicMaterialMediaStreamBridge/Private/DMMaterialValueMediaStreamDynamic.cpp
@@ -219,7 +219,37 @@ void UDMMaterialValueMediaStreamDynamic::OnComponentRemoved()
 
 void UDMMaterialValueMediaStreamDynamic::CopyParametersFrom_Implementation(UObject* InOther)
 {
-	// Overridden to do nothing
+	UDMMaterialValueMediaStream* Other = Cast<UDMMaterialValueMediaStream>(InOther);
+
+	UMediaStream* OtherMediaStream = Other->GetMediaStream();
+
+	if (!MediaStream || !OtherMediaStream)
+	{
+		return;
+	}
+
+	IMediaStreamPlayer* OtherPlayer = OtherMediaStream->GetPlayer().GetInterface();
+
+	if (OtherPlayer && !OtherPlayer->IsReadOnly())
+	{
+		if (IMediaStreamPlayer* Player = MediaStream->GetPlayer().GetInterface())
+		{
+			if (OtherPlayer->GetTextureConfig() != Player->GetTextureConfig())
+			{
+				OtherPlayer->SetTextureConfig(Player->GetTextureConfig());
+			}
+
+			if (OtherPlayer->GetPlayerConfig() != Player->GetPlayerConfig())
+			{
+				OtherPlayer->SetPlayerConfig(Player->GetPlayerConfig());
+			}
+		}
+	}
+
+	if (OtherMediaStream->GetSource() != MediaStream->GetSource())
+	{
+		OtherMediaStream->SetSource(MediaStream->GetSource());
+	}
 }
 
 #undef LOCTEXT_NAMESPACE
diff --git a/Engine/Plugins/Experimental/MediaStream/Source/MediaStream/Private/MediaStreamPlayerConfig.cpp b/Engine/Plugins/Experimental/MediaStream/Source/MediaStream/Private/MediaStreamPlayerConfig.cpp
--- a/Engine/Plugins/Experimental/MediaStream/Source/MediaStream/Private/MediaStreamPlayerConfig.cpp
+++ b/Engine/Plugins/Experimental/MediaStream/Source/MediaStream/Private/MediaStreamPlayerConfig.cpp
@@ -110,7 +110,7 @@ void FMediaStreamPlayerConfig::ApplyConfig(UMediaPlayer& InMediaPlayer) const
 	else if (InMediaPlayer.GetRate() != Rate)
 	{
 		InMediaPlayer.Seek(0.f);
-		InMediaPlayer.PlayOnOpen = false;
+		InMediaPlayer.Play();
 		ApplyRate(InMediaPlayer);
 	}
 
