Reconciled offline work

diff --git a/Engine/Plugins/Experimental/DynamicMaterialMediaStreamBridge/Source/DynamicMaterialMediaStreamBridgeEditor/Private/DMMaterialValueMediaStreamPropertyRowGenerator.cpp b/Engine/Plugins/Experimental/DynamicMaterialMediaStreamBridge/Source/DynamicMaterialMediaStreamBridgeEditor/Private/DMMaterialValueMediaStreamPropertyRowGenerator.cpp
--- a/Engine/Plugins/Experimental/DynamicMaterialMediaStreamBridge/Source/DynamicMaterialMediaStreamBridgeEditor/Private/DMMaterialValueMediaStreamPropertyRowGenerator.cpp
+++ b/Engine/Plugins/Experimental/DynamicMaterialMediaStreamBridge/Source/DynamicMaterialMediaStreamBridgeEditor/Private/DMMaterialValueMediaStreamPropertyRowGenerator.cpp
@@ -65,12 +65,10 @@ void FDMMaterialValueMediaStreamPropertyRowGenerator::AddComponentProperties(FDM
 	if (bHasPlayer)
 	{
 		AddDetailsCategory(MediaStreamParams);
-		AddTextureCategory(MediaStreamParams, /* Preview */ true);
 		AddCacheCategory(MediaStreamParams);
 		AddPlayerProperty(MediaStreamParams);
 
 		AddControlCategory(MediaStreamParams, /* Preview */ false);
-		AddTextureCategory(MediaStreamParams, /* Preview */ false);
 	}
 }
 
