Media Stream: Player changes (as well as source changes) are now propagated.

diff --git a/Engine/Plugins/Experimental/DynamicMaterialMediaStreamBridge/Source/DynamicMaterialMediaStreamBridge/Private/DMMaterialValueMediaStream.cpp b/Engine/Plugins/Experimental/DynamicMaterialMediaStreamBridge/Source/DynamicMaterialMediaStreamBridge/Private/DMMaterialValueMediaStream.cpp
--- a/Engine/Plugins/Experimental/DynamicMaterialMediaStreamBridge/Source/DynamicMaterialMediaStreamBridge/Private/DMMaterialValueMediaStream.cpp
+++ b/Engine/Plugins/Experimental/DynamicMaterialMediaStreamBridge/Source/DynamicMaterialMediaStreamBridge/Private/DMMaterialValueMediaStream.cpp
@@ -201,6 +201,11 @@ void UDMMaterialValueMediaStream::PostEditChangeProperty(FPropertyChangedEvent&
 	UpdatePlayer();
 }
 
+void UDMMaterialValueMediaStream::PostEditChangeChainProperty(FPropertyChangedChainEvent& InPropertyChangedEvent)
+{
+	Super::PostEditChangeChainProperty(InPropertyChangedEvent);
+}
+
 void UDMMaterialValueMediaStream::PostLoad()
 {
 	Super::PostLoad();
@@ -229,11 +234,12 @@ void UDMMaterialValueMediaStream::UpdateTextureFromMediaStream()
 	SetValue(MediaTexture);
 }
 
-void UDMMaterialValueMediaStream::SubscribeToPlayerChange()
+void UDMMaterialValueMediaStream::SubscribeToEvents()
 {
 	if (IsValid(MediaStream))
 	{
 		MediaStream->GetOnSourceChanged().AddDynamic(this, &UDMMaterialValueMediaStream::OnSourceChanged);
+		MediaStream->GetOnPlayerChanged().AddDynamic(this, &UDMMaterialValueMediaStream::OnPlayerChanged);
 		SubscribedStreamWeak = MediaStream;
 	}
 	else
@@ -242,16 +248,18 @@ void UDMMaterialValueMediaStream::SubscribeToPlayerChange()
 	}
 }
 
-void UDMMaterialValueMediaStream::UnsubscribeFromPlayerChange()
+void UDMMaterialValueMediaStream::UnsubscribeFromEvents()
 {
 	if (UMediaStream* SubscribedStream = SubscribedStreamWeak.Get())
 	{
 		SubscribedStream->GetOnSourceChanged().RemoveAll(this);
+		SubscribedStream->GetOnPlayerChanged().RemoveAll(this);
 	}
 
 	if (IsValid(MediaStream))
 	{
 		MediaStream->GetOnSourceChanged().RemoveAll(this);
+		MediaStream->GetOnPlayerChanged().RemoveAll(this);
 	}
 }
 
@@ -260,10 +268,16 @@ void UDMMaterialValueMediaStream::OnSourceChanged(UMediaStream* InMediaStream)
 	UpdatePlayer();
 }
 
+void UDMMaterialValueMediaStream::OnPlayerChanged(UMediaStream* InMediaStream)
+{
+	UpdatePlayer();
+	Update(this, EDMUpdateType::Value);
+}
+
 void UDMMaterialValueMediaStream::UpdatePlayer()
 {
-	UnsubscribeFromPlayerChange();
-	SubscribeToPlayerChange();
+	UnsubscribeFromEvents();
+	SubscribeToEvents();
 	UpdateTextureFromMediaStream();
 }
 
@@ -315,7 +329,7 @@ void UDMMaterialValueMediaStream::OnComponentRemoved()
 {
 	Super::OnComponentRemoved();
 
-	UnsubscribeFromPlayerChange();
+	UnsubscribeFromEvents();
 
 	SetValue(nullptr);
 }
diff --git a/Engine/Plugins/Experimental/DynamicMaterialMediaStreamBridge/Source/DynamicMaterialMediaStreamBridge/Private/DMMaterialValueMediaStreamDynamic.cpp b/Engine/Plugins/Experimental/DynamicMaterialMediaStreamBridge/Source/DynamicMaterialMediaStreamBridge/Private/DMMaterialValueMediaStreamDynamic.cpp
--- a/Engine/Plugins/Experimental/DynamicMaterialMediaStreamBridge/Source/DynamicMaterialMediaStreamBridge/Private/DMMaterialValueMediaStreamDynamic.cpp
+++ b/Engine/Plugins/Experimental/DynamicMaterialMediaStreamBridge/Source/DynamicMaterialMediaStreamBridge/Private/DMMaterialValueMediaStreamDynamic.cpp
@@ -160,11 +160,12 @@ void UDMMaterialValueMediaStreamDynamic::UpdateTextureFromMediaStream()
 	SetValue(MediaTexture);
 }
 
-void UDMMaterialValueMediaStreamDynamic::SubscribeToPlayerChange()
+void UDMMaterialValueMediaStreamDynamic::SubscribeToEvents()
 {
 	if (IsValid(MediaStream))
 	{
 		MediaStream->GetOnSourceChanged().AddDynamic(this, &UDMMaterialValueMediaStreamDynamic::OnSourceChanged);
+		MediaStream->GetOnPlayerChanged().AddDynamic(this, &UDMMaterialValueMediaStreamDynamic::OnPlayerChanged);
 		SubscribedStreamWeak = MediaStream;
 	}
 	else
@@ -173,16 +174,18 @@ void UDMMaterialValueMediaStreamDynamic::SubscribeToPlayerChange()
 	}
 }
 
-void UDMMaterialValueMediaStreamDynamic::UnsubscribeFromPlayerChange()
+void UDMMaterialValueMediaStreamDynamic::UnsubscribeFromEvents()
 {
 	if (UMediaStream* SubscribedStream = SubscribedStreamWeak.Get())
 	{
 		SubscribedStream->GetOnSourceChanged().RemoveAll(this);
+		SubscribedStream->GetOnPlayerChanged().RemoveAll(this);
 	}
 
 	if (IsValid(MediaStream))
 	{
 		MediaStream->GetOnSourceChanged().RemoveAll(this);
+		MediaStream->GetOnPlayerChanged().RemoveAll(this);
 	}
 }
 
@@ -193,10 +196,17 @@ void UDMMaterialValueMediaStreamDynamic::OnSourceChanged(UMediaStream* InMediaSt
 	Update(this, EDMUpdateType::RefreshDetailView);
 }
 
+void UDMMaterialValueMediaStreamDynamic::OnPlayerChanged(UMediaStream* InMediaStream)
+{
+	UpdatePlayer();
+
+	Update(this, EDMUpdateType::Value);
+}
+
 void UDMMaterialValueMediaStreamDynamic::UpdatePlayer()
 {
-	UnsubscribeFromPlayerChange();
-	SubscribeToPlayerChange();
+	UnsubscribeFromEvents();
+	SubscribeToEvents();
 	UpdateTextureFromMediaStream();
 }
 
@@ -211,7 +221,7 @@ void UDMMaterialValueMediaStreamDynamic::OnComponentRemoved()
 {
 	Super::OnComponentRemoved();
 
-	UnsubscribeFromPlayerChange();
+	UnsubscribeFromEvents();
 
 	SetValue(nullptr);
 }
diff --git a/Engine/Plugins/Experimental/DynamicMaterialMediaStreamBridge/Source/DynamicMaterialMediaStreamBridge/Public/DMMaterialValueMediaStream.h b/Engine/Plugins/Experimental/DynamicMaterialMediaStreamBridge/Source/DynamicMaterialMediaStreamBridge/Public/DMMaterialValueMediaStream.h
--- a/Engine/Plugins/Experimental/DynamicMaterialMediaStreamBridge/Source/DynamicMaterialMediaStreamBridge/Public/DMMaterialValueMediaStream.h
+++ b/Engine/Plugins/Experimental/DynamicMaterialMediaStreamBridge/Source/DynamicMaterialMediaStreamBridge/Public/DMMaterialValueMediaStream.h
@@ -46,6 +46,7 @@ public:
 
 	//~ Begin UObject
 	DYNAMICMATERIALMEDIASTREAMBRIDGE_API virtual void PostEditChangeProperty(FPropertyChangedEvent& InPropertyChangedEvent) override;
+	DYNAMICMATERIALMEDIASTREAMBRIDGE_API virtual void PostEditChangeChainProperty(FPropertyChangedChainEvent& InPropertyChangedEvent) override;
 	DYNAMICMATERIALMEDIASTREAMBRIDGE_API virtual void PostLoad() override;
 	//~ End UObject
 #endif
@@ -60,13 +61,16 @@ protected:
 #if WITH_EDITOR
 	void UpdateTextureFromMediaStream();
 
-	void SubscribeToPlayerChange();
+	void SubscribeToEvents();
 
-	void UnsubscribeFromPlayerChange();
+	void UnsubscribeFromEvents();
 
 	UFUNCTION()
 	void OnSourceChanged(UMediaStream* InMediaStream);
 
+	UFUNCTION()
+	void OnPlayerChanged(UMediaStream* InMediaStream);
+
 	void UpdatePlayer();
 
 	//~ Begin IDMParameterContainer
diff --git a/Engine/Plugins/Experimental/DynamicMaterialMediaStreamBridge/Source/DynamicMaterialMediaStreamBridge/Public/DMMaterialValueMediaStreamDynamic.h b/Engine/Plugins/Experimental/DynamicMaterialMediaStreamBridge/Source/DynamicMaterialMediaStreamBridge/Public/DMMaterialValueMediaStreamDynamic.h
--- a/Engine/Plugins/Experimental/DynamicMaterialMediaStreamBridge/Source/DynamicMaterialMediaStreamBridge/Public/DMMaterialValueMediaStreamDynamic.h
+++ b/Engine/Plugins/Experimental/DynamicMaterialMediaStreamBridge/Source/DynamicMaterialMediaStreamBridge/Public/DMMaterialValueMediaStreamDynamic.h
@@ -56,13 +56,16 @@ protected:
 #if WITH_EDITOR
 	void UpdateTextureFromMediaStream();
 
-	void SubscribeToPlayerChange();
+	void SubscribeToEvents();
 
-	void UnsubscribeFromPlayerChange();
+	void UnsubscribeFromEvents();
 
 	UFUNCTION()
 	void OnSourceChanged(UMediaStream* InMediaStream);
 
+	UFUNCTION()
+	void OnPlayerChanged(UMediaStream* InMediaStream);
+
 	void UpdatePlayer();
 
 	//~ Begin UDMMaterialComponent
diff --git a/Engine/Plugins/Experimental/MediaStream/Source/MediaStream/Private/MediaStream.cpp b/Engine/Plugins/Experimental/MediaStream/Source/MediaStream/Private/MediaStream.cpp
--- a/Engine/Plugins/Experimental/MediaStream/Source/MediaStream/Private/MediaStream.cpp
+++ b/Engine/Plugins/Experimental/MediaStream/Source/MediaStream/Private/MediaStream.cpp
@@ -92,6 +92,11 @@ UMediaStream::FOnSourceChanged& UMediaStream::GetOnSourceChanged()
 	return OnSourceChanged;
 }
 
+UMediaStream::FOnPlayerChanged& UMediaStream::GetOnPlayerChanged()
+{
+	return OnPlayerChanged;
+}
+
 void UMediaStream::Close()
 {
 	SetSource(FMediaStreamSource());
diff --git a/Engine/Plugins/Experimental/MediaStream/Source/MediaStream/Private/Players/MediaStreamLocalPlayer.cpp b/Engine/Plugins/Experimental/MediaStream/Source/MediaStream/Private/Players/MediaStreamLocalPlayer.cpp
--- a/Engine/Plugins/Experimental/MediaStream/Source/MediaStream/Private/Players/MediaStreamLocalPlayer.cpp
+++ b/Engine/Plugins/Experimental/MediaStream/Source/MediaStream/Private/Players/MediaStreamLocalPlayer.cpp
@@ -108,6 +108,11 @@ void UMediaStreamLocalPlayer::PostEditChangeProperty(FPropertyChangedEvent& InPr
 	{
 		ApplyPlaylistIndex();
 	}
+	
+	if (UMediaStream* MediaStream = GetMediaStream())
+	{
+		MediaStream->GetOnPlayerChanged().Broadcast(MediaStream);
+	}
 }
 #endif
 
diff --git a/Engine/Plugins/Experimental/MediaStream/Source/MediaStream/Private/Players/MediaStreamProxyPlayer.cpp b/Engine/Plugins/Experimental/MediaStream/Source/MediaStream/Private/Players/MediaStreamProxyPlayer.cpp
--- a/Engine/Plugins/Experimental/MediaStream/Source/MediaStream/Private/Players/MediaStreamProxyPlayer.cpp
+++ b/Engine/Plugins/Experimental/MediaStream/Source/MediaStream/Private/Players/MediaStreamProxyPlayer.cpp
@@ -60,6 +60,18 @@ void UMediaStreamProxyPlayer::Deinitialize()
 	ProxyStream = nullptr;
 }
 
+#if WITH_EDITOR
+void UMediaStreamProxyPlayer::PostEditChangeProperty(FPropertyChangedEvent& InPropertyChangedEvent)
+{
+	Super::PostEditChangeProperty(InPropertyChangedEvent);
+
+	if (UMediaStream* MediaStream = GetMediaStream())
+	{
+		MediaStream->GetOnPlayerChanged().Broadcast(MediaStream);
+	}
+}
+#endif
+
 void UMediaStreamProxyPlayer::OnSourceChanged(const FMediaStreamSource& InSource)
 {
 	ProxyStreamSoft.Reset();
diff --git a/Engine/Plugins/Experimental/MediaStream/Source/MediaStream/Private/Players/MediaStreamProxyPlayer.h b/Engine/Plugins/Experimental/MediaStream/Source/MediaStream/Private/Players/MediaStreamProxyPlayer.h
--- a/Engine/Plugins/Experimental/MediaStream/Source/MediaStream/Private/Players/MediaStreamProxyPlayer.h
+++ b/Engine/Plugins/Experimental/MediaStream/Source/MediaStream/Private/Players/MediaStreamProxyPlayer.h
@@ -59,6 +59,12 @@ public:
 	virtual void Deinitialize() override;
 	//~ End IMediaStreamPlayer
 
+#if WITH_EDITOR
+	//~ Begin UObject
+	virtual void PostEditChangeProperty(FPropertyChangedEvent& InPropertyChangedEvent) override;
+	//~ End UObject
+#endif
+
 protected:
 	/** Soft pointer to the proxied stream. */
 	UPROPERTY()
diff --git a/Engine/Plugins/Experimental/MediaStream/Source/MediaStream/Public/MediaStream.h b/Engine/Plugins/Experimental/MediaStream/Source/MediaStream/Public/MediaStream.h
--- a/Engine/Plugins/Experimental/MediaStream/Source/MediaStream/Public/MediaStream.h
+++ b/Engine/Plugins/Experimental/MediaStream/Source/MediaStream/Public/MediaStream.h
@@ -69,11 +69,17 @@ public:
 	MEDIASTREAM_API bool EnsurePlayer(bool bInForceRecreatePlayer = false);
 
 	/**
-	 * Called when the Media Stream Player changes.
+	 * Called when the Media Stream Source changes.
 	 */
 	DECLARE_DYNAMIC_MULTICAST_DELEGATE_OneParam(FOnSourceChanged, UMediaStream*, InMediaStream);
 	MEDIASTREAM_API FOnSourceChanged& GetOnSourceChanged();
 
+	/**
+	 * Called when the Media Stream Player or its Settings change.
+	 */
+	DECLARE_DYNAMIC_MULTICAST_DELEGATE_OneParam(FOnPlayerChanged, UMediaStream*, InMediaStream);
+	MEDIASTREAM_API FOnPlayerChanged& GetOnPlayerChanged();
+
 	/**
 	 * Stops the media stream and unloads any resources.
 	 */
@@ -117,6 +123,9 @@ protected:
 	UPROPERTY(BlueprintAssignable)
 	FOnSourceChanged OnSourceChanged;
 
+	UPROPERTY(BlueprintAssignable)
+	FOnPlayerChanged OnPlayerChanged;
+
 	FMediaStreamSource Source_PreUndo;
 
 	/**
