Material Designer / Media Stream: When choosing a media layer, now sets the mask stage to a 'use base texture' texture source.

diff --git a/Engine/Plugins/Experimental/DynamicMaterialMediaStreamBridge/Source/DynamicMaterialMediaStreamBridgeEditor/Private/DMMediaStreamStageSourceMenuExtender.cpp b/Engine/Plugins/Experimental/DynamicMaterialMediaStreamBridge/Source/DynamicMaterialMediaStreamBridgeEditor/Private/DMMediaStreamStageSourceMenuExtender.cpp
--- a/Engine/Plugins/Experimental/DynamicMaterialMediaStreamBridge/Source/DynamicMaterialMediaStreamBridgeEditor/Private/DMMediaStreamStageSourceMenuExtender.cpp
+++ b/Engine/Plugins/Experimental/DynamicMaterialMediaStreamBridge/Source/DynamicMaterialMediaStreamBridgeEditor/Private/DMMediaStreamStageSourceMenuExtender.cpp
@@ -160,6 +160,31 @@ void FDMMediaStreamStageSourceMenuExtender::ChangeSourceToMediaStreamFromContext
 			0,
 			FDMMaterialStageConnectorChannel::THREE_CHANNELS
 		);
+
+		if (UDMMaterialLayerObject* Layer = Stage->GetLayer())
+		{
+			if (UDMMaterialStage* MaskStage = Layer->GetStage(EDMMaterialLayerStage::Mask))
+			{
+				if (UDMMaterialStageThroughputLayerBlend* MaskLayerBlend = Cast<UDMMaterialStageThroughputLayerBlend>(MaskStage->GetSource()))
+				{
+					if (UDMMaterialStageInputExpression* MaskInputExpression = Cast<UDMMaterialStageInputExpression>(MaskLayerBlend->GetInputMask()))
+					{
+						UDMMaterialStageExpressionTextureSample* MaskTextureSample = Cast<UDMMaterialStageExpressionTextureSample>(MaskInputExpression->GetMaterialStageExpression());
+
+						if (!MaskTextureSample)
+						{
+							MaskInputExpression->SetMaterialStageExpressionClass(UDMMaterialStageExpressionTextureSample::StaticClass());
+							MaskTextureSample = Cast<UDMMaterialStageExpressionTextureSample>(MaskInputExpression->GetMaterialStageExpression());
+						}
+
+						if (MaskTextureSample)
+						{
+							MaskTextureSample->SetUseBaseTexture(true);
+						}
+					}
+				}
+			}
+		}
 	}
 	else if (StageSource->IsA<UDMMaterialStageThroughputLayerBlend>())
 	{
@@ -282,6 +307,27 @@ void FDMMediaStreamStageSourceMenuExtender::AddMediaStreamLayerFromContext(UDMMe
 
 	NewLayer->Update(NewLayer, EDMUpdateType::Structure);
 
+	if (UDMMaterialStage* MaskStage = NewLayer->GetStage(EDMMaterialLayerStage::Mask))
+	{
+		if (UDMMaterialStageThroughputLayerBlend* MaskLayerBlend = Cast<UDMMaterialStageThroughputLayerBlend>(MaskStage->GetSource()))
+		{
+			if (UDMMaterialStageInputExpression* MaskInputExpression = Cast<UDMMaterialStageInputExpression>(MaskLayerBlend->GetInputMask()))
+			{
+				UDMMaterialStageExpressionTextureSample* MaskTextureSample = Cast<UDMMaterialStageExpressionTextureSample>(MaskInputExpression->GetMaterialStageExpression());
+
+				if (!MaskTextureSample)
+				{
+					MaskInputExpression->SetMaterialStageExpressionClass(UDMMaterialStageExpressionTextureSample::StaticClass());
+					MaskTextureSample = Cast<UDMMaterialStageExpressionTextureSample>(MaskInputExpression->GetMaterialStageExpression());
+				}
+
+				if (MaskTextureSample)
+				{
+					MaskTextureSample->SetUseBaseTexture(true);
+				}
+			}
+		}
+	}
 }
 
 FDMMediaStreamStageSourceMenuExtender& FDMMediaStreamStageSourceMenuExtender::Get()
