Material Designer: You can now save levels with Material Designer instances in them that container media sources.

diff --git a/Engine/Plugins/Experimental/DynamicMaterialMediaStreamBridge/Source/DynamicMaterialMediaStreamBridge/Private/DMMaterialValueMediaStream.cpp b/Engine/Plugins/Experimental/DynamicMaterialMediaStreamBridge/Source/DynamicMaterialMediaStreamBridge/Private/DMMaterialValueMediaStream.cpp
--- a/Engine/Plugins/Experimental/DynamicMaterialMediaStreamBridge/Source/DynamicMaterialMediaStreamBridge/Private/DMMaterialValueMediaStream.cpp
+++ b/Engine/Plugins/Experimental/DynamicMaterialMediaStreamBridge/Source/DynamicMaterialMediaStreamBridge/Private/DMMaterialValueMediaStream.cpp
@@ -140,11 +140,25 @@ UDMMaterialValueDynamic* UDMMaterialValueMediaStream::ToDynamic(UDynamicMaterial
 {
 	UDMMaterialValueMediaStreamDynamic* ValueDynamic = UDMMaterialValueDynamic::CreateValueDynamic<UDMMaterialValueMediaStreamDynamic>(InMaterialModelDynamic, this);
 
+	// Media stream value must be set to public so it can be referenced by the dynamic value.
+	SetFlags(RF_Public);
+
 	if (MediaStream)
 	{
+		// Media stream must be set to public so it can be referenced by the dynamic value.
+		MediaStream->SetFlags(RF_Public);
+
 		if (UMediaStream* MediaStreamDynamic = ValueDynamic->GetMediaStream())
 		{
-			MediaStreamDynamic->SetSource(MediaStream->GetSource());
+			const FMediaStreamSource& Source = MediaStream->GetSource();
+
+			// Source needs to be public so it can be referenced by the dynamic value.
+			if (Source.Object)
+			{
+				Source.Object->SetFlags(RF_Public);
+			}
+
+			MediaStreamDynamic->SetSource(Source);
 		}
 	}
 
